Fixed a rare crash in crash report processor (an unhandled exception if Process.Exited gets called more than once)
Also fixed an internal error caused by my addition ReportData.cs

#CodeRevew Bob.Tellez,Jaroslaw.Surowiec

diff --git a/Engine/Source/Programs/DotNETCommon/DotNETUtilities/LaunchProcess.cs b/Engine/Source/Programs/DotNETCommon/DotNETUtilities/LaunchProcess.cs
--- a/Engine/Source/Programs/DotNETCommon/DotNETUtilities/LaunchProcess.cs
+++ b/Engine/Source/Programs/DotNETCommon/DotNETUtilities/LaunchProcess.cs
@@ -66,6 +66,9 @@ namespace Tools.DotNETCommon.LaunchProcess
 			{
 				LaunchedProcess.CancelOutputRead();
 				LaunchedProcess.CancelErrorRead();
+
+				// Protect against multiple calls
+				CaptureMessage = null;
 			}
 
 			LaunchedProcess.EnableRaisingEvents = false;
