Disable code analysis and doc generation for DotNETCommon. Unnecessary stall when building projects, and we don't use the output anyway.

diff --git a/Engine/Source/Programs/DotNETCommon/DotNETUtilities/DotNETUtilities.csproj b/Engine/Source/Programs/DotNETCommon/DotNETUtilities/DotNETUtilities.csproj
--- a/Engine/Source/Programs/DotNETCommon/DotNETUtilities/DotNETUtilities.csproj
+++ b/Engine/Source/Programs/DotNETCommon/DotNETUtilities/DotNETUtilities.csproj
@@ -21,18 +21,20 @@
     <DefineConstants>DEBUG;TRACE</DefineConstants>
     <ErrorReport>prompt</ErrorReport>
     <WarningLevel>4</WarningLevel>
-    <DocumentationFile>..\..\..\..\Binaries\DotNET\DotNETUtilities.XML</DocumentationFile>
+    <DocumentationFile>
+    </DocumentationFile>
   </PropertyGroup>
   <PropertyGroup Condition="'$(Configuration)|$(Platform)' == 'Development|AnyCPU'">
     <OutputPath>..\..\..\..\Binaries\DotNET\</OutputPath>
     <DefineConstants>TRACE</DefineConstants>
-    <DocumentationFile>..\..\..\..\Binaries\DotNET\DotNETUtilities.XML</DocumentationFile>
+    <DocumentationFile>
+    </DocumentationFile>
     <Optimize>true</Optimize>
     <DebugType>pdbonly</DebugType>
     <PlatformTarget>AnyCPU</PlatformTarget>
     <ErrorReport>prompt</ErrorReport>
     <CodeAnalysisRuleSet>MinimumRecommendedRules.ruleset</CodeAnalysisRuleSet>
-    <RunCodeAnalysis>true</RunCodeAnalysis>
+    <RunCodeAnalysis>false</RunCodeAnalysis>
   </PropertyGroup>
   <ItemGroup>
     <Reference Include="System" />
