Standard stream redirection removed during envvar harvest.

diff --git a/Engine/Source/Programs/DotNETCommon/DotNETUtilities/HarvestEnvVars.cs b/Engine/Source/Programs/DotNETCommon/DotNETUtilities/HarvestEnvVars.cs
--- a/Engine/Source/Programs/DotNETCommon/DotNETUtilities/HarvestEnvVars.cs
+++ b/Engine/Source/Programs/DotNETCommon/DotNETUtilities/HarvestEnvVars.cs
@@ -95,17 +95,6 @@ namespace Tools.DotNETCommon.HarvestEnvVars
 				StartInfo.CreateNoWindow  = true;
 				StartInfo.UseShellExecute = false;
 
-				// @todo UWP: Win10-beta issue, sometimes the process hangs if standard I/O is redirected, and there's bytes
-				// in the buffer and there is no reader for those streams.  Loop back with a later Win10 build
-				bool bRedirectStreams = System.Environment.OSVersion.Platform != PlatformID.Win32NT || System.Environment.OSVersion.Version.Major < 10;
-				bRedirectStreams = false;
-				if (bRedirectStreams)
-				{
-					StartInfo.RedirectStandardOutput = true;
-					StartInfo.RedirectStandardError = true;
-					StartInfo.RedirectStandardInput = true;
-				}
-
 				// Override path variable if it can be particularly long
 				if (PathOverride == EPathOverride.User)
 				{
@@ -126,10 +115,6 @@ namespace Tools.DotNETCommon.HarvestEnvVars
 				{
 					// Start the process up and then wait for it to finish
 					BatchFileProcess.Start();
-					if (bRedirectStreams)
-					{
-						BatchFileProcess.StandardOutput.ReadToEnd();
-					}
 					BatchFileProcess.WaitForExit();
 				}
 				catch (Exception Ex)
