Fix for when paths don't have shortened versions.

diff --git a/Engine/Source/Programs/DotNETCommon/DotNETUtilities/HarvestEnvVars.cs b/Engine/Source/Programs/DotNETCommon/DotNETUtilities/HarvestEnvVars.cs
--- a/Engine/Source/Programs/DotNETCommon/DotNETUtilities/HarvestEnvVars.cs
+++ b/Engine/Source/Programs/DotNETCommon/DotNETUtilities/HarvestEnvVars.cs
@@ -61,13 +61,13 @@ namespace Tools.DotNETCommon.HarvestEnvVars
 				var EnvReaderBatchFileContent = new StringBuilder();
 
 				// Run 'vcvars32.bat' (or similar x64 version) to set environment variables
-				EnvReaderBatchFileContent.AppendFormat("call {0} {1}", ShortBatchFileName, BatchFileParameters).AppendLine();
+				EnvReaderBatchFileContent.AppendFormat("call \"{0}\" {1}", ShortBatchFileName, BatchFileParameters).AppendLine();
 
 				// Pipe all environment variables to a file where we can read them in.
 				// We use a separate executable which runs after the batch file because we want to capture
 				// the environment after it has been set, and there's no easy way of doing this, and parsing
 				// the output of the set command is problematic when the vars contain non-ASCII characters.
-				EnvReaderBatchFileContent.AppendFormat("{0} {1}", ShortEnvVarsToXMLExePath, ShortEnvOutputFileName).AppendLine();
+				EnvReaderBatchFileContent.AppendFormat("\"{0}\" \"{1}\"", ShortEnvVarsToXMLExePath, ShortEnvOutputFileName).AppendLine();
 
 				FileInfo TempFileInfo = new FileInfo(EnvReaderBatchFileName);
 				if (TempFileInfo.Exists)
