Use Unicode version of GetShortPathName.

diff --git a/Engine/Source/Programs/DotNETCommon/DotNETUtilities/FileSystem.cs b/Engine/Source/Programs/DotNETCommon/DotNETUtilities/FileSystem.cs
--- a/Engine/Source/Programs/DotNETCommon/DotNETUtilities/FileSystem.cs
+++ b/Engine/Source/Programs/DotNETCommon/DotNETUtilities/FileSystem.cs
@@ -8,7 +8,7 @@ namespace Tools.DotNETCommon.FileSystem
 {
 	public static class FileSystem
 	{
-		[DllImport("kernel32.dll", SetLastError=true)]
+		[DllImport("kernel32.dll", SetLastError=true, CharSet=CharSet.Unicode)]
 		private static extern int GetShortPathName(string pathName, StringBuilder shortName, int cbShortName);
 
 		public static string GetShortPathName(string Path)
