Use %~dp0% to refer to current folder, avoiding non-ASCII path name issues where a short name doesn't exist.

diff --git a/Engine/Source/Programs/DotNETCommon/DotNETUtilities/HarvestEnvVars.cs b/Engine/Source/Programs/DotNETCommon/DotNETUtilities/HarvestEnvVars.cs
--- a/Engine/Source/Programs/DotNETCommon/DotNETUtilities/HarvestEnvVars.cs
+++ b/Engine/Source/Programs/DotNETCommon/DotNETUtilities/HarvestEnvVars.cs
@@ -56,7 +56,7 @@ namespace Tools.DotNETCommon.HarvestEnvVars
 
 				// Convert every path to short filenames to ensure we don't accidentally write out a non-ASCII batch file
 				var ShortBatchFileName       = FileSystem.FileSystem.GetShortPathName(BatchFileName);
-				var ShortEnvOutputFileName   = FileSystem.FileSystem.GetShortPathName(EnvOutputFileName);
+				var ShortEnvOutputFileName   = Path.Combine("%~dp0%", Path.GetFileName(EnvOutputFileName));
 				var ShortEnvVarsToXMLExePath = FileSystem.FileSystem.GetShortPathName(EnvVarsToXMLExePath);
 
 				var EnvReaderBatchFileContent = new StringBuilder();
