Crash report processor improvements - Using crash context when available (OR-4990/Automatically copy full dumps to Network Share) (still wip, but quite close to the end)

diff --git a/Engine/Source/Programs/DotNETCommon/DotNETUtilities/SimpleWebRequest.cs b/Engine/Source/Programs/DotNETCommon/DotNETUtilities/SimpleWebRequest.cs
--- a/Engine/Source/Programs/DotNETCommon/DotNETUtilities/SimpleWebRequest.cs
+++ b/Engine/Source/Programs/DotNETCommon/DotNETUtilities/SimpleWebRequest.cs
@@ -73,7 +73,7 @@ namespace Tools.DotNETCommon.SimpleWebRequest
 			{
 				HttpWebResponse WebResponse = ( HttpWebResponse )Request.GetResponse();
 				// Simple size check to prevent exploits
-				if( WebResponse.StatusCode == HttpStatusCode.OK && WebResponse.ContentLength < 1024 )
+				if( WebResponse.StatusCode == HttpStatusCode.OK && WebResponse.ContentLength < 4096 )
 				{
 					// Process the response
 					Stream ResponseStream = WebResponse.GetResponseStream();
diff --git a/Engine/Source/Programs/DotNETCommon/DotNETUtilities/XmlHandler.cs b/Engine/Source/Programs/DotNETCommon/DotNETUtilities/XmlHandler.cs
--- a/Engine/Source/Programs/DotNETCommon/DotNETUtilities/XmlHandler.cs
+++ b/Engine/Source/Programs/DotNETCommon/DotNETUtilities/XmlHandler.cs
@@ -51,7 +51,7 @@ namespace Tools.DotNETCommon.XmlHandler
 		/// <param name="Event">Details about the unknown attribute encountered in the Xml file.</param>
 		static private void XmlSerializer_UnknownAttribute( object Sender, XmlAttributeEventArgs Event )
 		{
-			Debug.WriteLine( " ... unknown XML attribute found '" + Event.Attr.Name + "' at line " + Event.LineNumber + " position " + Event.LinePosition );
+			//Debug.WriteLine( "Unknown XML attribute found '" + Event.Attr.Name + "' at line " + Event.LineNumber + " position " + Event.LinePosition );
 		}
 
 		/// <summary>
@@ -61,7 +61,7 @@ namespace Tools.DotNETCommon.XmlHandler
 		/// <param name="Event">Details about the unknown element found in the Xml file.</param>
 		static private void XmlSerializer_UnknownElement( object Sender, XmlElementEventArgs Event )
 		{
-			Debug.WriteLine( " ... unknown XML element found '" + Event.Element.Name + "' at line " + Event.LineNumber + " position " + Event.LinePosition );
+			//Debug.WriteLine( "Unknown XML element found '" + Event.Element.Name + "' at line " + Event.LineNumber + " position " + Event.LinePosition );
 		}
 
 		/// <summary>
@@ -71,7 +71,7 @@ namespace Tools.DotNETCommon.XmlHandler
 		/// <param name="Event">Details about the unknown node found in the Xml file.</param> 
 		static private void XmlSerializer_UnknownNode( object Sender, XmlNodeEventArgs Event )
 		{
-			Debug.WriteLine( " ... unknown XML node found '" + Event.Name + "' at line " + Event.LineNumber + " position " + Event.LinePosition );
+			//Debug.WriteLine( "Unknown XML node found '" + Event.Name + "' at line " + Event.LineNumber + " position " + Event.LinePosition );
 		}
 
 		/// <summary>
