Update copyright notices in Version.h and Metadata.cs.

diff --git a/Engine/Source/Programs/DotNETCommon/MetaData.cs b/Engine/Source/Programs/DotNETCommon/MetaData.cs
--- a/Engine/Source/Programs/DotNETCommon/MetaData.cs
+++ b/Engine/Source/Programs/DotNETCommon/MetaData.cs
@@ -1,4 +1,4 @@
-// Copyright 1998-2015 Epic Games, Inc. All Rights Reserved.
+// Copyright 1998-2016 Epic Games, Inc. All Rights Reserved.
 
 using System.Reflection;
 using System.Runtime.InteropServices;
@@ -7,7 +7,7 @@ using System.Resources;
 // These are the assembly properties for all tools
 [assembly: AssemblyCompany( "Epic Games, Inc." )]
 [assembly: AssemblyProduct( "UE4" )]
-[assembly: AssemblyCopyright( "Copyright 1998-2015 Epic Games, Inc. All Rights Reserved." )]
+[assembly: AssemblyCopyright( "Copyright 1998-2016 Epic Games, Inc. All Rights Reserved." )]
 [assembly: AssemblyTrademark( "" )]
 
 // Use a neutral culture to avoid some localisation issues
diff --git a/Engine/Source/Runtime/Launch/Resources/Version.h b/Engine/Source/Runtime/Launch/Resources/Version.h
--- a/Engine/Source/Runtime/Launch/Resources/Version.h
+++ b/Engine/Source/Runtime/Launch/Resources/Version.h
@@ -1,4 +1,4 @@
-// Copyright 1998-2015 Epic Games, Inc. All Rights Reserved.
+// Copyright 1998-2016 Epic Games, Inc. All Rights Reserved.
 
 #pragma once
 
@@ -14,7 +14,7 @@
 #define ENGINE_IS_PROMOTED_BUILD (BUILT_FROM_CHANGELIST > 0)
 
 #define EPIC_COMPANY_NAME  "Epic Games, Inc."
-#define EPIC_COPYRIGHT_STRING "Copyright 1998-2015 Epic Games, Inc. All Rights Reserved."
+#define EPIC_COPYRIGHT_STRING "Copyright 1998-2016 Epic Games, Inc. All Rights Reserved."
 #define EPIC_PRODUCT_NAME "Unreal Engine"
 #define EPIC_PRODUCT_IDENTIFIER "UnrealEngine"
 
