UBT: Fix intermittent PDB errors when compiling with ParallelExecutor, caused by MSPDBSRV.EXE instances being terminated while generating multiple PDBs in parallel.

diff --git a/Engine/Source/Programs/AutomationTool/AutomationUtils/ParallelExecutor.cs b/Engine/Source/Programs/AutomationTool/AutomationUtils/ParallelExecutor.cs
--- a/Engine/Source/Programs/AutomationTool/AutomationUtils/ParallelExecutor.cs
+++ b/Engine/Source/Programs/AutomationTool/AutomationUtils/ParallelExecutor.cs
@@ -38,6 +38,7 @@ namespace AutomationTool
 
 		class BuildActionExecutor
 		{
+			public ManagedProcessGroup ProcessGroup;
 			public BuildAction Action;
 			public List<string> LogLines = new List<string>();
 			public int ExitCode = -1;
@@ -45,8 +46,9 @@ namespace AutomationTool
 			private AutoResetEvent CompletedEvent;
 			private List<BuildActionExecutor> CompletedActions;
 
-			public BuildActionExecutor(BuildAction InAction, AutoResetEvent InCompletedEvent, List<BuildActionExecutor> InCompletedActions)
+			public BuildActionExecutor(ManagedProcessGroup InProcessGroup, BuildAction InAction, AutoResetEvent InCompletedEvent, List<BuildActionExecutor> InCompletedActions)
 			{
+				ProcessGroup = InProcessGroup;
 				Action = InAction;
 				CompletedEvent = InCompletedEvent;
 				CompletedActions = InCompletedActions;
@@ -59,7 +61,7 @@ namespace AutomationTool
 					LogLines.Add(Action.OutputPrefix);
 				}
 
-				using(ManagedProcess Process = new ManagedProcess(Action.ToolPath, Action.ToolArguments, Action.WorkingDirectory, Action.Environment, null, ManagedProcessPriority.BelowNormal))
+				using(ManagedProcess Process = new ManagedProcess(ProcessGroup, Action.ToolPath, Action.ToolArguments, Action.WorkingDirectory, Action.Environment, null, ManagedProcessPriority.BelowNormal))
 				{
 					LogLines.AddRange(Process.ReadAllLines());
 					ExitCode = Process.ExitCode;
@@ -103,83 +105,86 @@ namespace AutomationTool
 				Dictionary<BuildActionExecutor, Thread> ExecutingActions = new Dictionary<BuildActionExecutor,Thread>();
 				List<BuildActionExecutor> CompletedActions = new List<BuildActionExecutor>();
 
-				using(AutoResetEvent CompletedEvent = new AutoResetEvent(false))
+				using(ManagedProcessGroup ProcessGroup = new ManagedProcessGroup())
 				{
-					while(QueuedActions.Count > 0 || ExecutingActions.Count > 0)
+					using(AutoResetEvent CompletedEvent = new AutoResetEvent(false))
 					{
-						// Sort the actions by the number of things dependent on them
-						QueuedActions.Sort((A, B) => (A.TotalDependants == B.TotalDependants)? (B.SortIndex - A.SortIndex) : (B.TotalDependants - A.TotalDependants));
-
-						// Create threads up to the maximum number of actions
-						while(ExecutingActions.Count < MaxProcesses && QueuedActions.Count > 0)
+						while(QueuedActions.Count > 0 || ExecutingActions.Count > 0)
 						{
-							BuildAction Action = QueuedActions[QueuedActions.Count - 1];
-							QueuedActions.RemoveAt(QueuedActions.Count - 1);
+							// Sort the actions by the number of things dependent on them
+							QueuedActions.Sort((A, B) => (A.TotalDependants == B.TotalDependants)? (B.SortIndex - A.SortIndex) : (B.TotalDependants - A.TotalDependants));
 
-							BuildActionExecutor ExecutingAction = new BuildActionExecutor(Action, CompletedEvent, CompletedActions);
+							// Create threads up to the maximum number of actions
+							while(ExecutingActions.Count < MaxProcesses && QueuedActions.Count > 0)
+							{
+								BuildAction Action = QueuedActions[QueuedActions.Count - 1];
+								QueuedActions.RemoveAt(QueuedActions.Count - 1);
 
-							Thread ExecutingThread = new Thread(() => { ExecutingAction.Run(); });
-							ExecutingThread.Name = String.Format("Build:{0}", Action.Caption);
-							ExecutingThread.Start();
+								BuildActionExecutor ExecutingAction = new BuildActionExecutor(ProcessGroup, Action, CompletedEvent, CompletedActions);
 
-							ExecutingActions.Add(ExecutingAction, ExecutingThread);
-						}
+								Thread ExecutingThread = new Thread(() => { ExecutingAction.Run(); });
+								ExecutingThread.Name = String.Format("Build:{0}", Action.Caption);
+								ExecutingThread.Start();
 
-						// Wait for something to finish
-						CompletedEvent.WaitOne();
+								ExecutingActions.Add(ExecutingAction, ExecutingThread);
+							}
 
-						// Wait for something to finish and flush it to the log
-						lock(CompletedActions)
-						{
-							foreach(BuildActionExecutor CompletedAction in CompletedActions)
-							{
-								// Join the thread
-								Thread CompletedThread = ExecutingActions[CompletedAction];
-								CompletedThread.Join();
-								ExecutingActions.Remove(CompletedAction);
+							// Wait for something to finish
+							CompletedEvent.WaitOne();
 
-								// Write it to the log
-								if(CompletedAction.LogLines.Count > 0)
+							// Wait for something to finish and flush it to the log
+							lock(CompletedActions)
+							{
+								foreach(BuildActionExecutor CompletedAction in CompletedActions)
 								{
-									if(CurrentPrefix != CompletedAction.Action.GroupPrefix)
-									{
-										CurrentPrefix = CompletedAction.Action.GroupPrefix;
-										CommandUtils.Log(CurrentPrefix);
-									}
-									foreach(string LogLine in CompletedAction.LogLines)
+									// Join the thread
+									Thread CompletedThread = ExecutingActions[CompletedAction];
+									CompletedThread.Join();
+									ExecutingActions.Remove(CompletedAction);
+
+									// Write it to the log
+									if(CompletedAction.LogLines.Count > 0)
 									{
-										CommandUtils.Log(LogLine);
+										if(CurrentPrefix != CompletedAction.Action.GroupPrefix)
+										{
+											CurrentPrefix = CompletedAction.Action.GroupPrefix;
+											CommandUtils.Log(CurrentPrefix);
+										}
+										foreach(string LogLine in CompletedAction.LogLines)
+										{
+											CommandUtils.Log(LogLine);
+										}
 									}
-								}
 
-								// Check the exit code
-								if(CompletedAction.ExitCode == 0)
-								{
-									// Mark all the dependents as done
-									foreach(BuildAction DependantAction in CompletedAction.Action.Dependants)
+									// Check the exit code
+									if(CompletedAction.ExitCode == 0)
 									{
-										if(--DependantAction.MissingDependencyCount == 0)
+										// Mark all the dependents as done
+										foreach(BuildAction DependantAction in CompletedAction.Action.Dependants)
 										{
-											QueuedActions.Add(DependantAction);
+											if(--DependantAction.MissingDependencyCount == 0)
+											{
+												QueuedActions.Add(DependantAction);
+											}
 										}
 									}
-								}
-								else
-								{
-									// Update the exit code if it's not already set
-									if(ExitCode == 0)
+									else
 									{
-										ExitCode = CompletedAction.ExitCode;
+										// Update the exit code if it's not already set
+										if(ExitCode == 0)
+										{
+											ExitCode = CompletedAction.ExitCode;
+										}
 									}
 								}
+								CompletedActions.Clear();
 							}
-							CompletedActions.Clear();
-						}
 
-						// If we've already got a non-zero exit code, clear out the list of queued actions so nothing else will run
-						if(ExitCode != 0 && bStopOnErrors)
-						{
-							QueuedActions.Clear();
+							// If we've already got a non-zero exit code, clear out the list of queued actions so nothing else will run
+							if(ExitCode != 0 && bStopOnErrors)
+							{
+								QueuedActions.Clear();
+							}
 						}
 					}
 				}
diff --git a/Engine/Source/Programs/DotNETCommon/DotNETUtilities/ManagedProcess.cs b/Engine/Source/Programs/DotNETCommon/DotNETUtilities/ManagedProcess.cs
--- a/Engine/Source/Programs/DotNETCommon/DotNETUtilities/ManagedProcess.cs
+++ b/Engine/Source/Programs/DotNETCommon/DotNETUtilities/ManagedProcess.cs
@@ -24,53 +24,12 @@
 	}
 
 	/// <summary>
-	/// Encapsulates a managed child process, from which we can read the console output.
-	/// Uses job objects to ensure that the process will be terminated automatically by the O/S if the current process is terminated, and polls pipe reads to avoid blocking on unmanaged code
-	/// if the calling thread is terminated. Currently only implemented for Windows; makes heavy use of P/Invoke.
+	/// Tracks a set of processes, and destroys them when the object is disposed.
 	/// </summary>
-	public sealed class ManagedProcess : IDisposable
+	public class ManagedProcessGroup : IDisposable
 	{
-		[StructLayout(LayoutKind.Sequential)]
-		class SECURITY_ATTRIBUTES
-		{
-			public int nLength;
-			public IntPtr lpSecurityDescriptor;
-			public int bInheritHandle;
-		}
-
-		[StructLayout(LayoutKind.Sequential)]
-		class PROCESS_INFORMATION
-		{
-			public IntPtr hProcess;
-			public IntPtr hThread;
-			public uint dwProcessId;
-			public uint dwThreadId;
-		}
-
-		[StructLayout(LayoutKind.Sequential)]
-		class STARTUPINFO
-		{
-			public int cb;
-			public string lpReserved;
-			public string lpDesktop;
-			public string lpTitle;
-			public uint dwX;
-			public uint dwY;
-			public uint dwXSize;
-			public uint dwYSize;
-			public uint dwXCountChars;
-			public uint dwYCountChars;
-			public uint dwFillAttribute;
-			public uint dwFlags;
-			public short wShowWindow;
-			public short cbReserved2;
-			public IntPtr lpReserved2;
-			public SafeHandle hStdInput;
-			public SafeHandle hStdOutput;
-			public SafeHandle hStdError;
-		}
-
 		const UInt32 JOB_OBJECT_LIMIT_KILL_ON_JOB_CLOSE = 0x00002000;
+		const UInt32 JOB_OBJECT_LIMIT_BREAKAWAY_OK = 0x00000800;
 
 		[StructLayout(LayoutKind.Sequential)]
 		struct JOBOBJECT_BASIC_LIMIT_INFORMATION
@@ -116,6 +75,102 @@
         [DllImport("kernel32.dll", SetLastError=true)]
         static extern int SetInformationJobObject(SafeFileHandle hJob, int JobObjectInfoClass, IntPtr lpJobObjectInfo, int cbJobObjectInfoLength);
 
+		/// <summary>
+		/// Handle to the native job object that this process is added to. This handle is closed by the Dispose() method (and will automatically be closed by the OS on process exit),
+		/// resulting in the child process being killed.
+		/// </summary>
+		public SafeFileHandle JobHandle
+		{
+			get;
+			private set;
+		}
+
+		/// <summary>
+		/// Constructor
+		/// </summary>
+		public ManagedProcessGroup()
+		{
+			// Create the job object that the child process will be added to
+			JobHandle = CreateJobObject(IntPtr.Zero, IntPtr.Zero);
+			if(JobHandle == null)
+			{
+				throw new Win32Exception();
+			}
+
+			// Configure the job object to terminate the processes added to it when the handle is closed
+			JOBOBJECT_EXTENDED_LIMIT_INFORMATION LimitInformation = new JOBOBJECT_EXTENDED_LIMIT_INFORMATION();
+			LimitInformation.BasicLimitInformation.LimitFlags = JOB_OBJECT_LIMIT_KILL_ON_JOB_CLOSE | JOB_OBJECT_LIMIT_BREAKAWAY_OK;
+
+			int Length = Marshal.SizeOf(typeof(JOBOBJECT_EXTENDED_LIMIT_INFORMATION));
+			IntPtr LimitInformationPtr = Marshal.AllocHGlobal(Length);
+			Marshal.StructureToPtr(LimitInformation, LimitInformationPtr, false);
+
+			if(SetInformationJobObject(JobHandle, JobObjectExtendedLimitInformation, LimitInformationPtr, Length) == 0)
+			{
+				throw new Win32Exception();
+			}
+		}
+
+		/// <summary>
+		/// Dispose of the process group
+		/// </summary>
+		public void Dispose()
+		{
+			if(JobHandle != null)
+			{
+				JobHandle.Dispose();
+				JobHandle = null;
+			}
+		}
+	}
+
+	/// <summary>
+	/// Encapsulates a managed child process, from which we can read the console output.
+	/// Uses job objects to ensure that the process will be terminated automatically by the O/S if the current process is terminated, and polls pipe reads to avoid blocking on unmanaged code
+	/// if the calling thread is terminated. Currently only implemented for Windows; makes heavy use of P/Invoke.
+	/// </summary>
+	public sealed class ManagedProcess : IDisposable
+	{
+		[StructLayout(LayoutKind.Sequential)]
+		class SECURITY_ATTRIBUTES
+		{
+			public int nLength;
+			public IntPtr lpSecurityDescriptor;
+			public int bInheritHandle;
+		}
+
+		[StructLayout(LayoutKind.Sequential)]
+		class PROCESS_INFORMATION
+		{
+			public IntPtr hProcess;
+			public IntPtr hThread;
+			public uint dwProcessId;
+			public uint dwThreadId;
+		}
+
+		[StructLayout(LayoutKind.Sequential)]
+		class STARTUPINFO
+		{
+			public int cb;
+			public string lpReserved;
+			public string lpDesktop;
+			public string lpTitle;
+			public uint dwX;
+			public uint dwY;
+			public uint dwXSize;
+			public uint dwYSize;
+			public uint dwXCountChars;
+			public uint dwYCountChars;
+			public uint dwFillAttribute;
+			public uint dwFlags;
+			public short wShowWindow;
+			public short cbReserved2;
+			public IntPtr lpReserved2;
+			public SafeHandle hStdInput;
+			public SafeHandle hStdOutput;
+			public SafeHandle hStdError;
+		}
+
 		[DllImport("kernel32.dll", SetLastError=true)]
 		static extern int AssignProcessToJobObject(SafeFileHandle hJob, IntPtr hProcess);
 
@@ -204,11 +259,5 @@
 		StreamReader ReadStream;
 
 		/// <summary>
-		/// Handle to the native job object that this process is added to. This handle is closed by the Dispose() method (and will automatically be closed by the OS on process exit),
-		/// resulting in the child process being killed.
-		/// </summary>
-		SafeFileHandle JobHandle;
-
-		/// <summary>
 		/// Static lock object. This is used to synchronize the creation of child processes - in particular, the inheritance of stdout/stderr write pipes. If processes
 		/// inherit pipes meant for other processes, they won't be closed until both terminate.
@@ -218,33 +267,14 @@
 		/// <summary>
 		/// Spawns a new managed process.
 		/// </summary>
+		/// <param name="Group">The managed process group to add to</param>
 		/// <param name="FileName">Path to the executable to be run</param>
 		/// <param name="CommandLine">Command line arguments for the process</param>
 		/// <param name="WorkingDirectory">Working directory for the new process. May be null to use the current working directory.</param>
 		/// <param name="Environment">Environment variables for the new process. May be null, in which case the current process' environment is inherited</param>
 		/// <param name="Input">Text to be passed via stdin to the new process. May be null.</param>
-		public ManagedProcess(string FileName, string CommandLine, string WorkingDirectory, IReadOnlyDictionary<string, string> Environment, string Input, ManagedProcessPriority Priority)
+		public ManagedProcess(ManagedProcessGroup Group, string FileName, string CommandLine, string WorkingDirectory, IReadOnlyDictionary<string, string> Environment, string Input, ManagedProcessPriority Priority)
 		{
-			// Create the job object that the child process will be added to
-			JobHandle = CreateJobObject(IntPtr.Zero, IntPtr.Zero);
-			if(JobHandle == null)
-			{
-				throw new Win32Exception();
-			}
-
-			// Configure the job object to terminate the processes added to it when the handle is closed
-			JOBOBJECT_EXTENDED_LIMIT_INFORMATION LimitInformation = new JOBOBJECT_EXTENDED_LIMIT_INFORMATION();
-			LimitInformation.BasicLimitInformation.LimitFlags = JOB_OBJECT_LIMIT_KILL_ON_JOB_CLOSE;
-
-			int Length = Marshal.SizeOf(typeof(JOBOBJECT_EXTENDED_LIMIT_INFORMATION));
-			IntPtr LimitInformationPtr = Marshal.AllocHGlobal(Length);
-			Marshal.StructureToPtr(LimitInformation, LimitInformationPtr, false);
-
-			if(SetInformationJobObject(JobHandle, JobObjectExtendedLimitInformation, LimitInformationPtr, Length) == 0)
-			{
-				throw new Win32Exception();
-			}
-
 			// Create the child process
 			IntPtr EnvironmentBlock = IntPtr.Zero;
 			try
@@ -349,7 +379,7 @@
 					}
 
 					// Add it to our job object
-					if(AssignProcessToJobObject(JobHandle, ProcessInfo.hProcess) == 0)
+					if(AssignProcessToJobObject(Group.JobHandle, ProcessInfo.hProcess) == 0)
 					{
 						// Support for nested job objects was only addeed in Windows 8; prior to that, assigning processes to job objects would fail. Figure out if we're already in a job, and ignore the error if we are.
 						int OriginalError = Marshal.GetLastWin32Error();
@@ -431,11 +461,6 @@
 				StdOutRead.Dispose();
 				StdOutRead = null;
 			}
-			if(JobHandle != null)
-			{
-				JobHandle.Dispose();
-				JobHandle = null;
-			}
 		}
 
 		/// <summary>
diff --git a/Engine/Source/Programs/UnrealBuildTool/System/ParallelExecutor.cs b/Engine/Source/Programs/UnrealBuildTool/System/ParallelExecutor.cs
--- a/Engine/Source/Programs/UnrealBuildTool/System/ParallelExecutor.cs
+++ b/Engine/Source/Programs/UnrealBuildTool/System/ParallelExecutor.cs
@@ -16,7 +16,7 @@ namespace UnrealBuildTool
 	/// </summary>
 	class ParallelExecutor : ActionExecutor
 	{
-		[DebuggerDisplay("{Caption}")]
+		[DebuggerDisplay("{Inner}")]
 		class BuildAction
 		{
 			public int SortIndex;
@@ -145,83 +145,86 @@ namespace UnrealBuildTool
 				Dictionary<BuildAction, Thread> ExecutingActions = new Dictionary<BuildAction,Thread>();
 				List<BuildAction> CompletedActions = new List<BuildAction>();
 
-				using(AutoResetEvent CompletedEvent = new AutoResetEvent(false))
+				using(ManagedProcessGroup ProcessGroup = new ManagedProcessGroup())
 				{
-					int NumCompletedActions = 0;
-					using (ProgressWriter ProgressWriter = new ProgressWriter("Compiling C++ source code...", false))
+					using(AutoResetEvent CompletedEvent = new AutoResetEvent(false))
 					{
-						while(QueuedActions.Count > 0 || ExecutingActions.Count > 0)
+						int NumCompletedActions = 0;
+						using (ProgressWriter ProgressWriter = new ProgressWriter("Compiling C++ source code...", false))
 						{
-							// Sort the actions by the number of things dependent on them
-							QueuedActions.Sort((A, B) => (A.TotalDependantCount == B.TotalDependantCount)? (B.SortIndex - A.SortIndex) : (B.TotalDependantCount - A.TotalDependantCount));
-
-							// Create threads up to the maximum number of actions
-							while(ExecutingActions.Count < MaxProcesses && QueuedActions.Count > 0)
+							while(QueuedActions.Count > 0 || ExecutingActions.Count > 0)
 							{
-								BuildAction Action = QueuedActions[QueuedActions.Count - 1];
-								QueuedActions.RemoveAt(QueuedActions.Count - 1);
+								// Sort the actions by the number of things dependent on them
+								QueuedActions.Sort((A, B) => (A.TotalDependantCount == B.TotalDependantCount)? (B.SortIndex - A.SortIndex) : (B.TotalDependantCount - A.TotalDependantCount));
 
-								Thread ExecutingThread = new Thread(() => { ExecuteAction(Action, CompletedActions, CompletedEvent); });
-								ExecutingThread.Name = String.Format("Build:{0}", Action.Inner.StatusDescription);
-								ExecutingThread.Start();
+								// Create threads up to the maximum number of actions
+								while(ExecutingActions.Count < MaxProcesses && QueuedActions.Count > 0)
+								{
+									BuildAction Action = QueuedActions[QueuedActions.Count - 1];
+									QueuedActions.RemoveAt(QueuedActions.Count - 1);
 
-								ExecutingActions.Add(Action, ExecutingThread);
-							}
+									Thread ExecutingThread = new Thread(() => { ExecuteAction(ProcessGroup, Action, CompletedActions, CompletedEvent); });
+									ExecutingThread.Name = String.Format("Build:{0}", Action.Inner.StatusDescription);
+									ExecutingThread.Start();
 
-							// Wait for something to finish
-							CompletedEvent.WaitOne();
+									ExecutingActions.Add(Action, ExecutingThread);
+								}
 
-							// Wait for something to finish and flush it to the log
-							lock(CompletedActions)
-							{
-								foreach(BuildAction CompletedAction in CompletedActions)
-								{
-									// Join the thread
-									Thread CompletedThread = ExecutingActions[CompletedAction];
-									CompletedThread.Join();
-									ExecutingActions.Remove(CompletedAction);
+								// Wait for something to finish
+								CompletedEvent.WaitOne();
 
-									// Write it to the log
-									if(CompletedAction.LogLines.Count > 0)
+								// Wait for something to finish and flush it to the log
+								lock(CompletedActions)
+								{
+									foreach(BuildAction CompletedAction in CompletedActions)
 									{
-										foreach(string LogLine in CompletedAction.LogLines)
+										// Join the thread
+										Thread CompletedThread = ExecutingActions[CompletedAction];
+										CompletedThread.Join();
+										ExecutingActions.Remove(CompletedAction);
+
+										// Write it to the log
+										if(CompletedAction.LogLines.Count > 0)
 										{
-											Log.TraceInformation(LogLine);
+											foreach(string LogLine in CompletedAction.LogLines)
+											{
+												Log.TraceInformation(LogLine);
+											}
 										}
-									}
 
-									// Update the progress
-									NumCompletedActions++;
-									ProgressWriter.Write(NumCompletedActions, InputActions.Count);
+										// Update the progress
+										NumCompletedActions++;
+										ProgressWriter.Write(NumCompletedActions, InputActions.Count);
 
-									// Check the exit code
-									if(CompletedAction.ExitCode == 0)
-									{
-										// Mark all the dependents as done
-										foreach(BuildAction DependantAction in CompletedAction.Dependants)
+										// Check the exit code
+										if(CompletedAction.ExitCode == 0)
 										{
-											if(--DependantAction.MissingDependencyCount == 0)
+											// Mark all the dependents as done
+											foreach(BuildAction DependantAction in CompletedAction.Dependants)
 											{
-												QueuedActions.Add(DependantAction);
+												if(--DependantAction.MissingDependencyCount == 0)
+												{
+													QueuedActions.Add(DependantAction);
+												}
 											}
 										}
-									}
-									else
-									{
-										// Update the exit code if it's not already set
-										if(bResult && CompletedAction.ExitCode != 0)
+										else
 										{
-											bResult = false;
+											// Update the exit code if it's not already set
+											if(bResult && CompletedAction.ExitCode != 0)
+											{
+												bResult = false;
+											}
 										}
 									}
+									CompletedActions.Clear();
 								}
-								CompletedActions.Clear();
-							}
 
-							// If we've already got a non-zero exit code, clear out the list of queued actions so nothing else will run
-							if(!bResult && bStopCompilationAfterErrors)
-							{
-								QueuedActions.Clear();
+								// If we've already got a non-zero exit code, clear out the list of queued actions so nothing else will run
+								if(!bResult && bStopCompilationAfterErrors)
+								{
+									QueuedActions.Clear();
+								}
 							}
 						}
 					}
@@ -234,17 +237,18 @@ namespace UnrealBuildTool
 		/// <summary>
 		/// Execute an individual action
 		/// </summary>
+		/// <param name="ProcessGroup">The process group</param>
 		/// <param name="Action">The action to execute</param>
 		/// <param name="CompletedActions">On completion, the list to add the completed action to</param>
 		/// <param name="CompletedEvent">Event to set once an event is complete</param>
-		static void ExecuteAction(BuildAction Action, List<BuildAction> CompletedActions, AutoResetEvent CompletedEvent)
+		static void ExecuteAction(ManagedProcessGroup ProcessGroup, BuildAction Action, List<BuildAction> CompletedActions, AutoResetEvent CompletedEvent)
 		{
 			if (Action.Inner.bShouldOutputStatusDescription && !String.IsNullOrEmpty(Action.Inner.StatusDescription))
 			{
 				Action.LogLines.Add(Action.Inner.StatusDescription);
 			}
 
-			using(ManagedProcess Process = new ManagedProcess(Action.Inner.CommandPath, Action.Inner.CommandArguments, Action.Inner.WorkingDirectory, null, null, ManagedProcessPriority.BelowNormal))
+			using(ManagedProcess Process = new ManagedProcess(ProcessGroup, Action.Inner.CommandPath, Action.Inner.CommandArguments, Action.Inner.WorkingDirectory, null, null, ManagedProcessPriority.BelowNormal))
 			{
 				Action.LogLines.AddRange(Process.ReadAllLines());
 				Action.ExitCode = Process.ExitCode;
