Change the Log class to be initialized to sensible defaults, allow settings to be modified at any time, and remove the requirement to manually initialize it.

#rnx

diff --git a/Engine/Source/Programs/AutomationTool/Program.cs b/Engine/Source/Programs/AutomationTool/Program.cs
--- a/Engine/Source/Programs/AutomationTool/Program.cs
+++ b/Engine/Source/Programs/AutomationTool/Program.cs
@@ -35,31 +35,21 @@ namespace AutomationTool
             }
 
 			// Parse the log level argument
-			LogEventType LogLevel = LogEventType.Log;
 			if(CommandUtils.ParseParam(Arguments, "-Verbose"))
 			{
-				LogLevel = LogEventType.Verbose;
+				Log.OutputLevel = LogEventType.Verbose;
 			}
 			if(CommandUtils.ParseParam(Arguments, "-VeryVerbose"))
 			{
-				LogLevel = LogEventType.VeryVerbose;
+				Log.OutputLevel = LogEventType.VeryVerbose;
 			}
 
 			// Initialize the log system, buffering the output until we can create the log file
 			StartupTraceListener StartupListener = new StartupTraceListener();
-			Log.InitLogging(
-                bLogTimestamps: CommandUtils.ParseParam(Arguments, "-Timestamps"),
-				InLogLevel: LogLevel,
-                bLogSeverity: true,
-				bLogProgramNameWithSeverity: false,
-                bLogSources: true,
-				bLogSourcesToConsole: false,
-                bColorConsoleOutput: true,
-                TraceListeners: new TraceListener[]
-                {
-					StartupListener
-				}
-			);
+			Trace.Listeners.Add(StartupListener);
+
+			// Configure log timestamps
+			Log.IncludeTimestamps = CommandUtils.ParseParam(Arguments, "Timestamps");
 
 			// Enter the main program section
             ExitCode ReturnCode = ExitCode.Success;
diff --git a/Engine/Source/Programs/DotNETCommon/DotNETUtilities/Log.cs b/Engine/Source/Programs/DotNETCommon/DotNETUtilities/Log.cs
--- a/Engine/Source/Programs/DotNETCommon/DotNETUtilities/Log.cs
+++ b/Engine/Source/Programs/DotNETCommon/DotNETUtilities/Log.cs
@@ -8,7 +8,6 @@ using System.Linq;
 using System.Reflection;
 using System.Runtime.CompilerServices;
 using System.Text;
-using System.Threading.Tasks;
 
 namespace Tools.DotNETCommon
 {
@@ -111,12 +110,6 @@ namespace Tools.DotNETCommon
 			public bool bHasFlushedHeadingText;
 		}
 
-		/// <summary>
-		/// Guard our initialization. Mainly used by top level exception handlers to ensure its safe to call a logging function.
-		/// In general user code should not concern itself checking for this.
-		/// </summary>
-		private static bool bIsInitialized = false;
-
 		/// <summary>
 		/// Object used for synchronization
 		/// </summary>
@@ -125,52 +118,76 @@ namespace Tools.DotNETCommon
 		/// <summary>
 		/// When true, verbose logging is enabled.
 		/// </summary>
-		private static LogEventType LogLevel = LogEventType.VeryVerbose;
+		public static LogEventType OutputLevel
+		{
+			get; set;
+		}
+
+		/// <summary>
+		/// Whether to include timestamps on each line of log output
+		/// </summary>
+		public static bool IncludeTimestamps
+		{
+			get; set;
+		}
 
 		/// <summary>
 		/// When true, warnings and errors will have a WARNING: or ERROR: prexifx, respectively.
 		/// </summary>
-		private static bool bLogSeverity = false;
+		public static bool IncludeSeverityPrefix
+		{
+			get; set;
+		}
 		
 		/// <summary>
 		/// When true, warnings and errors will have a prefix suitable for display by MSBuild (avoiding error messages showing as (EXEC : Error : ")
 		/// </summary>
-		private static bool bLogProgramNameWithSeverity = false;
+		public static bool IncludeProgramNameWithSeverityPrefix
+		{
+			get; set;
+		}
 
 		/// <summary>
 		/// When true, logs will have the calling mehod prepended to the output as MethodName:
 		/// </summary>
-		private static bool bLogSources = false;
+		public static bool IncludeCallingMethod
+		{
+			get; set;
+		}
 		
 		/// <summary>
 		/// When true, console output will have the calling mehod prepended to the output as MethodName:
 		/// </summary>
-		private static bool bLogSourcesToConsole = false;
+		public static bool IncludeCallingMethodForConsole
+		{
+			get; set;
+		}
 		
 		/// <summary>
 		/// When true, will detect warnings and errors and set the console output color to yellow and red.
 		/// </summary>
-		private static bool bColorConsoleOutput = false;
+		public static bool ColorConsoleOutput
+		{
+			get; set;
+		}
 
 		/// <summary>
 		/// Whether console output is redirected. This prevents writing status updates that rely on moving the cursor.
 		/// </summary>
-		private static bool bAllowStatusUpdates = true;
+		private static bool AllowStatusUpdates
+		{
+			get { return !Console.IsOutputRedirected; }
+		}
 
 		/// <summary>
 		/// When configured, this tracks time since initialization to prepend a timestamp to each log.
 		/// </summary>
-		private static Stopwatch Timer;
-
-		/// <summary>
-		/// Expose the log level. This is a hack for ProcessResult.LogOutput, which wants to bypass our normal formatting scheme.
-		/// </summary>
-		public static bool bIsVerbose { get { return LogLevel >= LogEventType.Verbose; } }
+		private static Stopwatch Timer = Stopwatch.StartNew();
 
 		/// <summary>
 		/// A collection of strings that have been already written once
 		/// </summary>
-		private static List<string> WriteOnceSet = new List<string>();
+		private static HashSet<string> WriteOnceSet = new HashSet<string>();
 
 		/// <summary>
 		/// Stack of status scope information.
@@ -192,8 +209,7 @@ namespace Tools.DotNETCommon
 		/// </summary>
 		public static string Indent
 		{
-			get;
-			set;
+			get; set;
 		}
 
 		/// <summary>
@@ -202,76 +218,12 @@ namespace Tools.DotNETCommon
 		static Log()
 		{
 			Indent = "";
-		}
-
-		/// <summary>
-		/// Allows code to check if the log system is ready yet.
-		/// End users should NOT need to use this. It pretty much exists
-		/// to work around startup issues since this is a global singleton.
-		/// </summary>
-		/// <returns></returns>
-		public static bool IsInitialized()
-		{
-			return bIsInitialized;
-		}
-
-		/// <summary>
-		/// Allows code to check if the log system is using console output color.
-		/// </summary>
-		/// <returns></returns>
-		public static bool ColorConsoleOutput()
-		{
-			return bColorConsoleOutput;
-		}
-
-		/// <summary>
-		/// Allows us to change verbosity after initializing. This can happen since we initialize logging early, 
-		/// but then read the config and command line later, which could change this value.
-		/// </summary>
-		public static void SetLoggingLevel(LogEventType InLogLevel)
-		{
-			Log.LogLevel = InLogLevel;
-		}
-
-		/// <summary>
-		/// This class allows InitLogging to be called more than once to work around chicken and eggs issues with logging and parsing command lines (see UBT startup code).
-		/// </summary>
-		/// <param name="bLogTimestamps">If true, the timestamp from Log init time will be prepended to all logs.</param>
-		/// <param name="InLogLevel"></param>
-		/// <param name="bLogSeverity">If true, warnings and errors will have a WARNING: and ERROR: prefix to them. </param>
-		/// <param name="bLogProgramNameWithSeverity">If true, includes the program name with any severity prefix</param>
-		/// <param name="bLogSources">If true, logs will have the originating method name prepended to them.</param>
-		/// <param name="bLogSourcesToConsole">If true, console output will have the originating method name appended to it.</param>
-		/// <param name="bColorConsoleOutput"></param>
-		/// <param name="TraceListeners">Collection of trace listeners to attach to the Trace.Listeners, in addition to the Default listener. The existing listeners (except the Default listener) are cleared first.</param>
-		public static void InitLogging(bool bLogTimestamps, LogEventType InLogLevel, bool bLogSeverity, bool bLogProgramNameWithSeverity, bool bLogSources, bool bLogSourcesToConsole, bool bColorConsoleOutput, IEnumerable<TraceListener> TraceListeners)
-		{
-			bIsInitialized = true;
-			Timer = (bLogTimestamps && Timer == null) ? Stopwatch.StartNew() : null;
-			Log.LogLevel = InLogLevel;
-			Log.bLogSeverity = bLogSeverity;
-			Log.bLogProgramNameWithSeverity = bLogProgramNameWithSeverity;
-			Log.bLogSources = bLogSources;
-			Log.bLogSourcesToConsole = bLogSourcesToConsole;
-			Log.bColorConsoleOutput = bColorConsoleOutput;
-			Log.bAllowStatusUpdates = !Console.IsOutputRedirected;
-
-			// ensure that if InitLogging is called more than once we don't stack listeners.
-			// but always leave the default listener around.
-			for (int ListenerNdx = 0; ListenerNdx < Trace.Listeners.Count;)
-			{
-				if (Trace.Listeners[ListenerNdx].GetType() != typeof(DefaultTraceListener))
-				{
-					Trace.Listeners.RemoveAt(ListenerNdx);
-				}
-				else
-				{
-					++ListenerNdx;
-				}
-			}
-			// don't add any null listeners
-			Trace.Listeners.AddRange(TraceListeners.Where(l => l != null).ToArray());
-			Trace.AutoFlush = true;
+			OutputLevel = LogEventType.Log;
+			IncludeSeverityPrefix = true;
+			IncludeProgramNameWithSeverityPrefix = false;
+			IncludeCallingMethod = true;
+			IncludeCallingMethodForConsole = false;
+			ColorConsoleOutput = true;
 		}
 
 		/// <summary>
@@ -314,16 +266,6 @@ namespace Tools.DotNETCommon
 			}
 		}
 
-		/// <summary>
-		/// Converts a LogEventType into a message code
-		/// </summary>
-		/// <param name="Severity"></param>
-		/// <returns></returns>
-		private static int GetMessageCode(LogEventType Severity)
-		{
-			return (int)Severity;
-		}
-
 		/// <summary>
 		/// Formats message for logging. Enforces the configured options.
 		/// </summary>
@@ -337,12 +279,12 @@ namespace Tools.DotNETCommon
 		[MethodImplAttribute(MethodImplOptions.NoInlining)]
 		private static List<string> FormatMessage(int StackFramesToSkip, LogEventType Verbosity, LogFormatOptions Options, bool bForConsole, string Format, params object[] Args)
 		{
-			string TimePrefix = (Timer != null) ? String.Format("[{0:hh\\:mm\\:ss\\.fff}] ", Timer.Elapsed) : "";
-			string SourcePrefix = (bForConsole ? bLogSourcesToConsole : bLogSources) ? string.Format("{0}: ", GetSource(StackFramesToSkip)) : "";
-			string SeverityPrefix = (bLogSeverity && ((Options & LogFormatOptions.NoSeverityPrefix) == 0)) ? GetSeverityPrefix(Verbosity) : "";
+			string TimePrefix = IncludeTimestamps? String.Format("[{0:hh\\:mm\\:ss\\.fff}] ", Timer.Elapsed) : "";
+			string SourcePrefix = (bForConsole ? IncludeCallingMethodForConsole : IncludeCallingMethod) ? string.Format("{0}: ", GetSource(StackFramesToSkip)) : "";
+			string SeverityPrefix = (IncludeSeverityPrefix && ((Options & LogFormatOptions.NoSeverityPrefix) == 0)) ? GetSeverityPrefix(Verbosity) : "";
 
 			// Include the executable name when running inside MSBuild. If unspecified, MSBuild re-formats them with an "EXEC :" prefix.
-			if(SeverityPrefix.Length > 0 && bLogProgramNameWithSeverity)
+			if(SeverityPrefix.Length > 0 && IncludeProgramNameWithSeverityPrefix)
 			{
 				SeverityPrefix = String.Format("{0}: {1}", Path.GetFileNameWithoutExtension(Assembly.GetEntryAssembly().Location), SeverityPrefix);
 			}
@@ -383,11 +325,6 @@ namespace Tools.DotNETCommon
 		[MethodImplAttribute(MethodImplOptions.NoInlining)]
 		private static void WriteLinePrivate(int StackFramesToSkip, bool bWriteOnce, LogEventType Verbosity, LogFormatOptions FormatOptions, string Format, params object[] Args)
 		{
-			if (!bIsInitialized)
-			{
-				throw new InvalidOperationException("Tried to using Logging system before it was ready");
-			}
-
 			// if we want this message only written one time, check if it was already written out
 			if (bWriteOnce)
 			{
@@ -400,7 +337,7 @@ namespace Tools.DotNETCommon
 				WriteOnceSet.Add(Formatted);
 			}
 
-			if (Verbosity <= LogLevel)
+			if (Verbosity <= OutputLevel)
 			{
 				lock (SyncObject)
 				{
@@ -416,12 +353,12 @@ namespace Tools.DotNETCommon
 					}
 
 					// Handle the console output separately; we format things differently
-					if ((Verbosity != LogEventType.Log || LogLevel >= LogEventType.Verbose) && (FormatOptions & LogFormatOptions.NoConsoleOutput) == 0)
+					if ((Verbosity != LogEventType.Log || OutputLevel >= LogEventType.Verbose) && (FormatOptions & LogFormatOptions.NoConsoleOutput) == 0)
 					{
 						FlushStatusHeading();
 
 						bool bResetConsoleColor = false;
-						if (bColorConsoleOutput)
+						if (ColorConsoleOutput)
 						{
 							if(Verbosity == LogEventType.Warning)
 							{
@@ -451,7 +388,7 @@ namespace Tools.DotNETCommon
 							}
 						}
 
-						if(StatusMessageStack.Count > 0 && bAllowStatusUpdates)
+						if(StatusMessageStack.Count > 0 && AllowStatusUpdates)
 						{
 							SetStatusText(StatusMessageStack.Peek().CurrentText);
 						}
@@ -787,7 +724,7 @@ namespace Tools.DotNETCommon
 			if(StatusMessageStack.Count > 0)
 			{
 				StatusMessage CurrentStatus = StatusMessageStack.Peek();
-				if(CurrentStatus.HeadingText.Length > 0 && !CurrentStatus.bHasFlushedHeadingText && bAllowStatusUpdates)
+				if(CurrentStatus.HeadingText.Length > 0 && !CurrentStatus.bHasFlushedHeadingText && AllowStatusUpdates)
 				{
 					SetStatusText(CurrentStatus.HeadingText);
 					Console.WriteLine();
@@ -842,7 +779,7 @@ namespace Tools.DotNETCommon
 				StatusMessage CurrentStatusMessage = StatusMessageStack.Peek();
 				CurrentStatusMessage.CurrentText = Message;
 
-				if(bAllowStatusUpdates || StatusTimer.Elapsed.TotalSeconds > 10.0)
+				if(AllowStatusUpdates || StatusTimer.Elapsed.TotalSeconds > 10.0)
 				{
 					SetStatusText(Message);
 					StatusTimer.Restart();
@@ -892,7 +829,7 @@ namespace Tools.DotNETCommon
 					NumCommonChars++;
 				}
 
-				if(!bAllowStatusUpdates && NumCommonChars < StatusText.Length)
+				if(!AllowStatusUpdates && NumCommonChars < StatusText.Length)
 				{
 					// Prevent writing backspace characters if the console doesn't support it
 					Console.WriteLine();
diff --git a/Engine/Source/Programs/UnrealBuildTool/Platform/Linux/LinuxToolChain.cs b/Engine/Source/Programs/UnrealBuildTool/Platform/Linux/LinuxToolChain.cs
--- a/Engine/Source/Programs/UnrealBuildTool/Platform/Linux/LinuxToolChain.cs
+++ b/Engine/Source/Programs/UnrealBuildTool/Platform/Linux/LinuxToolChain.cs
@@ -566,7 +566,7 @@ namespace UnrealBuildTool
 					{
 						Result += " -fdiagnostics-format=msvc";     // make diagnostics compatible with MSVC when cross-compiling
 					}
-					else if (Log.ColorConsoleOutput())
+					else if (Log.ColorConsoleOutput)
 					{
 						Result += " -fcolor-diagnostics";
 					}
diff --git a/Engine/Source/Programs/UnrealBuildTool/UnrealBuildTool.cs b/Engine/Source/Programs/UnrealBuildTool/UnrealBuildTool.cs
--- a/Engine/Source/Programs/UnrealBuildTool/UnrealBuildTool.cs
+++ b/Engine/Source/Programs/UnrealBuildTool/UnrealBuildTool.cs
@@ -326,32 +326,21 @@ namespace UnrealBuildTool
 			ECompilationResult Result = ECompilationResult.Succeeded;
 
 			// Parse the log level argument
-			LogEventType LogLevel = LogEventType.Log;
 			if (Arguments.Any(x => x.Equals("-Verbose", StringComparison.InvariantCultureIgnoreCase)))
 			{
-				LogLevel = LogEventType.Verbose;
+				Log.OutputLevel = LogEventType.Verbose;
 			}
 			else if (Arguments.Any(x => x.Equals("-VeryVerbose", StringComparison.InvariantCultureIgnoreCase)))
 			{
-				LogLevel = LogEventType.VeryVerbose;
+				Log.OutputLevel = LogEventType.VeryVerbose;
 			}
 
 			// Initialize the log system, buffering the output until we can create the log file
 			StartupTraceListener StartupListener = new StartupTraceListener();
-			bool bLogProgramNameWithSeverity = Arguments.Any(x => x.Equals("-FromMsBuild", StringComparison.InvariantCultureIgnoreCase));
-			Log.InitLogging(
-				bLogTimestamps: Arguments.Any(x => x.Equals("-Timestamps", StringComparison.InvariantCultureIgnoreCase)),
-				InLogLevel: LogLevel,
-				bLogSeverity: true,
-				bLogProgramNameWithSeverity: bLogProgramNameWithSeverity,
-				bLogSources: true,
-				bLogSourcesToConsole: false,
-				bColorConsoleOutput: true,
-				TraceListeners: new TraceListener[]
-				{
-					StartupListener
-				}
-			);
+			Trace.Listeners.Add(StartupListener);
+
+			// If we're running from MsBuild, we need to include the program name to avoid it formatting messages as EXEC : UnrealBuildTool : ERROR
+			Log.IncludeProgramNameWithSeverityPrefix = Arguments.Any(x => x.Equals("-FromMsBuild", StringComparison.InvariantCultureIgnoreCase));
 
 			// Write the command line
 			Log.TraceLog("Command line: {0}", Environment.CommandLine);
@@ -504,7 +493,7 @@ namespace UnrealBuildTool
 					CommandLine.ParseArguments(Arguments, BuildConfiguration);
 
 					// Copy some of the static settings that are being deprecated from BuildConfiguration
-					bPrintDebugInfo = BuildConfiguration.bPrintDebugInfo || LogLevel == LogEventType.Verbose || LogLevel == LogEventType.VeryVerbose;
+					bPrintDebugInfo = BuildConfiguration.bPrintDebugInfo || Log.OutputLevel == LogEventType.Verbose || Log.OutputLevel == LogEventType.VeryVerbose;
 					bPrintPerformanceInfo = BuildConfiguration.bPrintPerformanceInfo;
 
 					// Don't run junk deleter if we're just setting up autosdks
@@ -952,10 +941,7 @@ namespace UnrealBuildTool
 			}
 			catch (Exception Exception)
 			{
-				if (Log.IsInitialized())
-				{
-					Log.TraceError("UnrealBuildTool Exception: " + Exception.ToString());
-				}
+				Log.TraceError("UnrealBuildTool Exception: " + Exception.ToString());
 				return (int)ECompilationResult.OtherCompilationError;
 			}
 		}
