Add convenience methods for creating a FileReference or DirectoryReference from a (possibly null) string.

#rnx

diff --git a/Engine/Source/Programs/DotNETCommon/DotNETUtilities/DirectoryReference.cs b/Engine/Source/Programs/DotNETCommon/DotNETUtilities/DirectoryReference.cs
--- a/Engine/Source/Programs/DotNETCommon/DotNETUtilities/DirectoryReference.cs
+++ b/Engine/Source/Programs/DotNETCommon/DotNETUtilities/DirectoryReference.cs
@@ -231,7 +231,31 @@ namespace Tools.DotNETCommon
 		/// <returns>The directory containing the file  </returns>
 		public static DirectoryReference FromFile(FileReference File)
 		{
-			return (File == null)? null : File.Directory;
+			if(File == null)
+			{
+				return null;
+			}
+			else
+			{
+				return File.Directory;
+			}
+		}
+
+		/// <summary>
+		/// Create a DirectoryReference from a string. If the string is null, returns a null DirectoryReference.
+		/// </summary>
+		/// <param name="DirectoryName">Path for the new object</param>
+		/// <returns>Returns a FileReference representing the given string, or null.</returns>
+		public static DirectoryReference FromString(string DirectoryName)
+		{
+			if(String.IsNullOrEmpty(DirectoryName))
+			{
+				return null;
+			}
+			else
+			{
+				return new DirectoryReference(DirectoryName);
+			}
 		}
 
 		#region System.IO.Directory Wrapper Methods
diff --git a/Engine/Source/Programs/DotNETCommon/DotNETUtilities/FileReference.cs b/Engine/Source/Programs/DotNETCommon/DotNETUtilities/FileReference.cs
--- a/Engine/Source/Programs/DotNETCommon/DotNETUtilities/FileReference.cs
+++ b/Engine/Source/Programs/DotNETCommon/DotNETUtilities/FileReference.cs
@@ -55,6 +55,23 @@ namespace Tools.DotNETCommon
 		{
 		}
 
+		/// <summary>
+		/// Create a FileReference from a string. If the string is null, returns a null FileReference.
+		/// </summary>
+		/// <param name="FileName">FileName for the string</param>
+		/// <returns>Returns a FileReference representing the given string, or null.</returns>
+		public static FileReference FromString(string FileName)
+		{
+			if(String.IsNullOrEmpty(FileName))
+			{
+				return null;
+			}
+			else
+			{
+				return new FileReference(FileName);
+			}
+		}
+
 		/// <summary>
 		/// Gets the file name without path information
 		/// </summary>
