UBT: Fix a couple of edge cases when parsing output from child processes.

* If a single line was larger than 32kb, the process would be treated as having terminated.
* If the last output line did not have a trailing newline, the line would be dropped.

diff --git a/Engine/Source/Programs/DotNETCommon/DotNETUtilities/ManagedProcess.cs b/Engine/Source/Programs/DotNETCommon/DotNETUtilities/ManagedProcess.cs
--- a/Engine/Source/Programs/DotNETCommon/DotNETUtilities/ManagedProcess.cs
+++ b/Engine/Source/Programs/DotNETCommon/DotNETUtilities/ManagedProcess.cs
@@ -575,10 +575,20 @@ namespace Tools.DotNETCommon
 			int NumBytesInBuffer = 0;
 			for(;;)
 			{
+				// If we're got a single line larger than 32kb (!), enlarge the buffer to ensure we can handle it
+				if(NumBytesInBuffer == Buffer.Length)
+				{
+					Array.Resize(ref Buffer, Buffer.Length + 32 * 1024);
+				}
+
 				// Fill the buffer, reentering managed code every 20ms to allow thread abort exceptions to be thrown
 				int NumBytesRead = Read(Buffer, NumBytesInBuffer, Buffer.Length - NumBytesInBuffer);
 				if(NumBytesRead == 0)
 				{
+					if(NumBytesInBuffer < Buffer.Length)
+					{
+						OutputLines.Add(Encoding.UTF8.GetString(Buffer, 0, NumBytesInBuffer));
+					}
 					break;
 				}
 
