UBT: Don't assume that console encoding is UTF-8 when using ParallelExecutor.

diff --git a/Engine/Source/Programs/DotNETCommon/DotNETUtilities/ManagedProcess.cs b/Engine/Source/Programs/DotNETCommon/DotNETUtilities/ManagedProcess.cs
--- a/Engine/Source/Programs/DotNETCommon/DotNETUtilities/ManagedProcess.cs
+++ b/Engine/Source/Programs/DotNETCommon/DotNETUtilities/ManagedProcess.cs
@@ -302,9 +302,9 @@ namespace Tools.DotNETCommon
 						List<byte> EnvironmentBytes = new List<byte>();
 						foreach(KeyValuePair<string, string> Pair in Environment)
 						{
-							EnvironmentBytes.AddRange(Encoding.UTF8.GetBytes(Pair.Key));
+							EnvironmentBytes.AddRange(Console.OutputEncoding.GetBytes(Pair.Key));
 							EnvironmentBytes.Add((byte)'=');
-							EnvironmentBytes.AddRange(Encoding.UTF8.GetBytes(Pair.Value));
+							EnvironmentBytes.AddRange(Console.OutputEncoding.GetBytes(Pair.Value));
 							EnvironmentBytes.Add((byte)0);
 						}
 						EnvironmentBytes.Add((byte)0);
@@ -587,7 +587,7 @@ namespace Tools.DotNETCommon
 				{
 					if(NumBytesInBuffer < Buffer.Length)
 					{
-						OutputLines.Add(Encoding.UTF8.GetString(Buffer, 0, NumBytesInBuffer));
+						OutputLines.Add(Console.OutputEncoding.GetString(Buffer, 0, NumBytesInBuffer));
 					}
 					break;
 				}
@@ -603,7 +603,7 @@ namespace Tools.DotNETCommon
 					{
 						if(Buffer[Idx] != '\n' || LastCharacter != '\r')
 						{
-							OutputLines.Add(Encoding.UTF8.GetString(Buffer, LastStartIdx, Idx - LastStartIdx));
+							OutputLines.Add(Console.OutputEncoding.GetString(Buffer, LastStartIdx, Idx - LastStartIdx));
 						}
 						LastStartIdx = Idx + 1;
 					}
