Output any printable errors from Perforce as text, rather than as a hex dump.

diff --git a/Engine/Source/Programs/DotNETCommon/DotNETUtilities/Perforce/PerforceChildProcess.cs b/Engine/Source/Programs/DotNETCommon/DotNETUtilities/Perforce/PerforceChildProcess.cs
--- a/Engine/Source/Programs/DotNETCommon/DotNETUtilities/Perforce/PerforceChildProcess.cs
+++ b/Engine/Source/Programs/DotNETCommon/DotNETUtilities/Perforce/PerforceChildProcess.cs
@@ -78,6 +78,60 @@ namespace Tools.DotNETCommon.Perforce
 			return bEndOfStream;
 		}
 
+		private string FormatUnexpectedData(byte NextByte)
+		{
+			// Read all the data to the buffer
+			List<byte> Data = new List<byte>();
+			for(int Idx = 0; Idx < 1024; Idx++)
+			{
+				Data.Add(NextByte);
+				if(IsEndOfStream())
+				{
+					break;
+				}
+				NextByte = ReadByte();
+			}
+
+			// Check if it's printable. 
+			bool bIsPrintable = true;
+			for(int Idx = 0; Idx < Data.Count; Idx++)
+			{
+				if(Data[Idx] < 0x09 || (Data[Idx] >= 0x0e && Data[Idx] <= 0x1f) || Data[Idx] == 0x1f)
+				{
+					bIsPrintable = false;
+				}
+			}
+
+			// Format the result
+			StringBuilder Result = new StringBuilder();
+			if(Data.Count > 0)
+			{
+				if(bIsPrintable)
+				{
+					Result.Append("\n    ");
+					for(int Idx = 0; Idx < Data.Count; Idx++)
+					{
+						if(Data[Idx] == '\n')
+						{
+							Result.Append("\n    ");
+						}
+						else if(Data[Idx] != '\r')
+						{
+							Result.Append((char)Data[Idx]);
+						}
+					}
+				}
+				else
+				{
+					for(int Idx = 0; Idx < Data.Count; Idx++)
+					{
+						Result.AppendFormat("{0}{1:x2}", ((Idx & 31) == 0)? "\n    " : " ", Data[Idx]);
+					}
+				}
+			}
+			return Result.ToString();
+		}
+
 		public bool TryReadRecord(List<KeyValuePair<string, object>> Record)
 		{
 			// Check if we've reached the end of the stream. This is the only condition where we return false.
@@ -90,17 +144,7 @@ namespace Tools.DotNETCommon.Perforce
 			byte Temp = ReadByte();
 			if(Temp != '{')
 			{
-				StringBuilder Result = new StringBuilder("Unexpected data while parsing marshaled output - expected '{', got:");
-				for(int Idx = 0; Idx < 1024; Idx++)
-				{
-					Result.AppendFormat("{0}{1:x2}", ((Idx & 31) == 0)? "\n    " : " ", Temp);
-					if(IsEndOfStream())
-					{
-						break;
-					}
-					Temp = ReadByte();
-				}
-				throw new PerforceException(Result.ToString());
+				throw new PerforceException("Unexpected data while parsing marshaled output - expected '{{', got: {0}", FormatUnexpectedData(Temp));
 			}
 
 			// Read all the fields in the record
@@ -115,7 +159,7 @@ namespace Tools.DotNETCommon.Perforce
 				}
 				else if(KeyFieldType != 's')
 				{
-					throw new PerforceException("Unexpected key field type while parsing marshalled output ({0}) - expected 's'", (int)KeyFieldType);
+					throw new PerforceException("Unexpected key field type while parsing marshalled output ({0}) - expected 's', got: {1}", (int)KeyFieldType, FormatUnexpectedData(KeyFieldType));
 				}
 
 				// Read the key
@@ -143,7 +187,7 @@ namespace Tools.DotNETCommon.Perforce
 				}
 				else
 				{
-					throw new PerforceException("Unexpected value field type while parsing marshalled output ({0}) - expected 's'", (int)ValueFieldType);
+					throw new PerforceException("Unexpected value field type while parsing marshalled output ({0}) - expected 's', got: {1}", (int)ValueFieldType, FormatUnexpectedData(ValueFieldType));
 				}
 			}
 			return true;
