Add a wrapper around the FastJSON library that registers custom Epic types, and formats objects in the same style that the engine uses.

diff --git a/Engine/Source/Programs/DotNETCommon/DotNETUtilities/DotNETUtilities.csproj b/Engine/Source/Programs/DotNETCommon/DotNETUtilities/DotNETUtilities.csproj
--- a/Engine/Source/Programs/DotNETCommon/DotNETUtilities/DotNETUtilities.csproj
+++ b/Engine/Source/Programs/DotNETCommon/DotNETUtilities/DotNETUtilities.csproj
@@ -72,7 +72,9 @@
     <Compile Include="FilePattern.cs" />
     <Compile Include="FileReference.cs" />
     <Compile Include="FileSystemReference.cs" />
+    <Compile Include="Json.cs" />
     <Compile Include="JsonObject.cs" />
+    <Compile Include="JsonWriter.cs" />
     <Compile Include="Log.cs" />
     <Compile Include="LogIndentScope.cs" />
     <Compile Include="LogStatusScope.cs" />
diff --git a/Engine/Source/Programs/DotNETCommon/DotNETUtilities/Json.cs b/Engine/Source/Programs/DotNETCommon/DotNETUtilities/Json.cs
new file mode 100644
--- /dev/null
+++ b/Engine/Source/Programs/DotNETCommon/DotNETUtilities/Json.cs
@@ -0,0 +1,227 @@
+using System;
+using System.Collections.Generic;
+using System.IO;
+using System.Linq;
+using System.Text;
+using System.Threading.Tasks;
+
+namespace Tools.DotNETCommon
+{
+	/// <summary>
+	/// Options for serializing an object to JSON
+	/// </summary>
+	public enum JsonSerializeOptions
+	{
+		/// <summary>
+		/// No options
+		/// </summary>
+		None,
+
+		/// <summary>
+		/// Whether to format the output for readability
+		/// </summary>
+		PrettyPrint = 1,
+	}
+
+	/// <summary>
+	/// Wrapper class around fastJSON methods.
+	/// </summary>
+	public static class Json
+	{
+		/// <summary>
+		/// Static constructor. Registers known types to the serializer.
+		/// </summary>
+		static Json()
+		{
+			fastJSON.JSON.Instance.RegisterCustomType(typeof(FileReference), x => EscapeString(x.ToString()), x => new FileReference(x));
+			fastJSON.JSON.Instance.RegisterCustomType(typeof(DirectoryReference), x => EscapeString(x.ToString()), x => new DirectoryReference(x));
+		}
+
+		/// <summary>
+		/// Serializes an object to JSON
+		/// </summary>
+		/// <param name="Object">Object to be serialized</param>
+		/// <returns>JSON representation of the object</returns>
+		public static string Serialize(object Object)
+		{
+			return Serialize(Object, JsonSerializeOptions.PrettyPrint);
+		}
+
+		/// <summary>
+		/// Serializes an object to JSON
+		/// </summary>
+		/// <param name="Object">Object to be serialized</param>
+		/// <param name="Options">Options for formatting the object</param>
+		/// <returns>JSON representation of the object</returns>
+		public static string Serialize(object Object, JsonSerializeOptions Options)
+		{
+			fastJSON.JSONParameters Params = new fastJSON.JSONParameters();
+			Params.EnableAnonymousTypes = true;
+
+			string Text = fastJSON.JSON.Instance.ToJSON(Object, Params);
+			if((Options & JsonSerializeOptions.PrettyPrint) != 0)
+			{
+				Text = Format(Text);
+			}
+			return Text;
+		}
+
+		/// <summary>
+		/// Deserialize an object from text
+		/// </summary>
+		/// <typeparam name="T">Type of the object to deserialize</typeparam>
+		/// <param name="Text">String to parse</param>
+		/// <returns>Constructed object</returns>
+		public static T Deserialize<T>(string Text)
+		{
+			return fastJSON.JSON.Instance.ToObject<T>(Text);
+		}
+
+		/// <summary>
+		/// Formats the given JSON string in 'Epic' format (parentheses/brackets on new lines, tabs instead of spaces, no space before colons in key/value pairs)
+		/// </summary>
+		/// <param name="Input">The input string</param>
+		/// <returns>The formatted string</returns>
+		public static string Format(string Input)
+		{
+			StringBuilder Result = new StringBuilder(Input.Length * 2);
+
+			bool bOnEmptyLine = true;
+
+			int Depth = 0;
+			for(int Idx = 0; Idx < Input.Length; Idx++)
+			{
+				switch(Input[Idx])
+				{
+					case '{':
+					case '[':
+						if(!bOnEmptyLine)
+						{
+							Result.AppendLine();
+							bOnEmptyLine = true;
+						}
+
+						AppendIndent(Result, Depth++);
+						Result.Append(Input[Idx]);
+						Result.AppendLine();
+						break;
+					case '}':
+					case ']':
+						if(!bOnEmptyLine)
+						{
+							Result.AppendLine();
+						}
+
+						AppendIndent(Result, --Depth);
+						Result.Append(Input[Idx]);
+						bOnEmptyLine = false;
+						break;
+					case ':':
+						Result.Append(Input[Idx]);
+						Result.Append(' ');
+						break;
+					case ',':
+						Result.Append(Input[Idx]);
+						Result.AppendLine();
+						bOnEmptyLine = true;
+						break;
+					default:
+						if(!Char.IsWhiteSpace(Input[Idx]))
+						{
+							// If this is an empty line, append the indent
+							if(bOnEmptyLine)
+							{
+								AppendIndent(Result, Depth);
+								bOnEmptyLine = false;
+							}
+
+							// Append the character
+							Result.Append(Input[Idx]);
+
+							// If this was a string, also append everything to the end of the string
+							if(Input[Idx] == '\"')
+							{
+								Idx++;
+								for(; Idx < Input.Length; Idx++)
+								{
+									Result.Append(Input[Idx]);
+									if(Input[Idx] == '\\' && Idx + 1 < Input.Length)
+									{
+										Result.Append(Input[++Idx]);
+									}
+									else if(Input[Idx] == '\"')
+									{
+										break;
+									}
+								}
+							}
+						}
+						break;
+				}
+			}
+
+			return Result.ToString();
+		}
+
+		/// <summary>
+		/// Appends an indent to the output string
+		/// </summary>
+		/// <param name="Builder">The string builder to append to</param>
+		/// <param name="Depth">The number of tabs to append</param>
+		private static void AppendIndent(StringBuilder Builder, int Depth)
+		{
+			for(int Idx = 0; Idx < Depth; Idx++)
+			{
+				Builder.Append('\t');
+			}
+		}
+
+		/// <summary>
+		/// Escapes a string for serializing to JSON
+		/// </summary>
+		/// <param name="Value">The string to escape</param>
+		/// <returns>The escaped string</returns>
+		public static string EscapeString(string Value)
+		{
+			StringBuilder Result = new StringBuilder();
+			for (int Idx = 0; Idx < Value.Length; Idx++)
+			{
+				switch (Value[Idx])
+				{
+					case '\"':
+						Result.Append("\\\"");
+						break;
+					case '\\':
+						Result.Append("\\\\");
+						break;
+					case '\b':
+						Result.Append("\\b");
+						break;
+					case '\f':
+						Result.Append("\\f");
+						break;
+					case '\n':
+						Result.Append("\\n");
+						break;
+					case '\r':
+						Result.Append("\\r");
+						break;
+					case '\t':
+						Result.Append("\\t");
+						break;
+					default:
+						if (Char.IsControl(Value[Idx]))
+						{
+							Result.AppendFormat("\\u{0:X4}", (int)Value[Idx]);
+						}
+						else
+						{
+							Result.Append(Value[Idx]);
+						}
+						break;
+				}
+			}
+			return Result.ToString();
+		}
+	}
+}
diff --git a/Engine/Source/Programs/DotNETCommon/DotNETUtilities/JsonObject.cs b/Engine/Source/Programs/DotNETCommon/DotNETUtilities/JsonObject.cs
--- a/Engine/Source/Programs/DotNETCommon/DotNETUtilities/JsonObject.cs
+++ b/Engine/Source/Programs/DotNETCommon/DotNETUtilities/JsonObject.cs
@@ -1,4 +1,4 @@
-﻿// Copyright 1998-2018 Epic Games, Inc. All Rights Reserved.
+// Copyright 1998-2018 Epic Games, Inc. All Rights Reserved.
 
 using System;
 using System.Collections.Generic;
@@ -437,279 +437,4 @@ namespace Tools.DotNETCommon
 			}
 		}
 	}
-
-	/// <summary>
-	/// Writer for JSON data, which indents the output text appropriately, and adds commas and newlines between fields
-	/// </summary>
-	public class JsonWriter : IDisposable
-	{
-		TextWriter Writer;
-		bool bLeaveOpen;
-		bool bRequiresComma;
-		string Indent;
-
-		/// <summary>
-		/// Constructor
-		/// </summary>
-		/// <param name="FileName">File to write to</param>
-		public JsonWriter(string FileName)
-			: this(new StreamWriter(FileName))
-		{
-		}
-
-		/// <summary>
-		/// Constructor
-		/// </summary>
-		/// <param name="FileName">File to write to</param>
-		public JsonWriter(FileReference FileName)
-			: this(new StreamWriter(FileName.FullName))
-		{
-		}
-
-		/// <summary>
-		/// Constructor
-		/// </summary>
-		/// <param name="Writer">The text writer to output to</param>
-		/// <param name="bLeaveOpen">Whether to leave the writer open when the object is disposed</param>
-		public JsonWriter(TextWriter Writer, bool bLeaveOpen = false)
-		{
-			this.Writer = Writer;
-			this.bLeaveOpen = bLeaveOpen;
-			Indent = "";
-		}
-
-		/// <summary>
-		/// Dispose of any managed resources
-		/// </summary>
-		public void Dispose()
-		{
-			if(!bLeaveOpen && Writer != null)
-			{
-				Writer.Dispose();
-				Writer = null;
-			}
-		}
-
-		/// <summary>
-		/// Write the opening brace for an object
-		/// </summary>
-		public void WriteObjectStart()
-		{
-			WriteCommaNewline();
-
-			Writer.Write(Indent);
-			Writer.Write("{");
-
-			Indent += "\t";
-			bRequiresComma = false;
-		}
-
-		/// <summary>
-		/// Write the name and opening brace for an object
-		/// </summary>
-		/// <param name="ObjectName">Name of the field</param>
-		public void WriteObjectStart(string ObjectName)
-		{
-			WriteCommaNewline();
-
-			Writer.Write("{0}\"{1}\": ", Indent, ObjectName);
-
-			bRequiresComma = false;
-
-			WriteObjectStart();
-		}
-
-		/// <summary>
-		/// Write the closing brace for an object
-		/// </summary>
-		public void WriteObjectEnd()
-		{
-			Indent = Indent.Substring(0, Indent.Length - 1);
-
-			Writer.WriteLine();
-			Writer.Write(Indent);
-			Writer.Write("}");
-
-			bRequiresComma = true;
-		}
-
-		/// <summary>
-		/// Write the name and opening bracket for an array
-		/// </summary>
-		/// <param name="ArrayName">Name of the field</param>
-		public void WriteArrayStart(string ArrayName)
-		{
-			WriteCommaNewline();
-
-			Writer.Write("{0}\"{1}\": [", Indent, ArrayName);
-
-			Indent += "\t";
-			bRequiresComma = false;
-		}
-
-		/// <summary>
-		/// Write the closing bracket for an array
-		/// </summary>
-		public void WriteArrayEnd()
-		{
-			Indent = Indent.Substring(0, Indent.Length - 1);
-
-			Writer.WriteLine();
-			Writer.Write("{0}]", Indent);
-
-			bRequiresComma = true;
-		}
-
-		/// <summary>
-		/// Write an array of strings
-		/// </summary>
-		/// <param name="Name">Name of the field</param>
-		/// <param name="Values">Values for the field</param>
-		public void WriteStringArrayField(string Name, IEnumerable<string> Values)
-		{
-			WriteArrayStart(Name);
-			foreach(string Value in Values)
-			{
-				WriteValue(Value);
-			}
-			WriteArrayEnd();
-		}
-
-		/// <summary>
-		/// Write an array of enum values
-		/// </summary>
-		/// <param name="Name">Name of the field</param>
-		/// <param name="Values">Values for the field</param>
-		public void WriteEnumArrayField<T>(string Name, IEnumerable<T> Values) where T : struct
-		{
-			WriteStringArrayField(Name, Values.Select(x => x.ToString()));
-		}
-
-		/// <summary>
-		/// Write a value with no field name, for the contents of an array
-		/// </summary>
-		/// <param name="Value">Value to write</param>
-		public void WriteValue(string Value)
-		{
-			WriteCommaNewline();
-
-			Writer.Write(Indent);
-			WriteEscapedString(Value);
-
-			bRequiresComma = true;
-		}
-
-		/// <summary>
-		/// Write a field name and string value
-		/// </summary>
-		/// <param name="Name">Name of the field</param>
-		/// <param name="Value">Value for the field</param>
-		public void WriteValue(string Name, string Value)
-		{
-			WriteCommaNewline();
-
-			Writer.Write("{0}\"{1}\": ", Indent, Name);
-			WriteEscapedString(Value);
-
-			bRequiresComma = true;
-		}
-
-		/// <summary>
-		/// Write a field name and integer value
-		/// </summary>
-		/// <param name="Name">Name of the field</param>
-		/// <param name="Value">Value for the field</param>
-		public void WriteValue(string Name, int Value)
-		{
-			WriteValueInternal(Name, Value.ToString());
-		}
-
-		/// <summary>
-		/// Write a field name and double value
-		/// </summary>
-		/// <param name="Name">Name of the field</param>
-		/// <param name="Value">Value for the field</param>
-		public void WriteValue(string Name, double Value)
-		{
-			WriteValueInternal(Name, Value.ToString());
-		}
-
-		/// <summary>
-		/// Write a field name and bool value
-		/// </summary>
-		/// <param name="Name">Name of the field</param>
-		/// <param name="Value">Value for the field</param>
-		public void WriteValue(string Name, bool Value)
-		{
-			WriteValueInternal(Name, Value ? "true" : "false");
-		}
-
-		void WriteCommaNewline()
-		{
-			if (bRequiresComma)
-			{
-				Writer.WriteLine(",");
-			}
-			else if (Indent.Length > 0)
-			{
-				Writer.WriteLine();
-			}
-		}
-
-		void WriteValueInternal(string Name, string Value)
-		{
-			WriteCommaNewline();
-
-			Writer.Write("{0}\"{1}\": {2}", Indent, Name, Value);
-
-			bRequiresComma = true;
-		}
-
-		void WriteEscapedString(string Value)
-		{
-			// Escape any characters which may not appear in a JSON string (see http://www.json.org).
-			Writer.Write("\"");
-			if (Value != null)
-			{
-				for (int Idx = 0; Idx < Value.Length; Idx++)
-				{
-					switch (Value[Idx])
-					{
-						case '\"':
-							Writer.Write("\\\"");
-							break;
-						case '\\':
-							Writer.Write("\\\\");
-							break;
-						case '\b':
-							Writer.Write("\\b");
-							break;
-						case '\f':
-							Writer.Write("\\f");
-							break;
-						case '\n':
-							Writer.Write("\\n");
-							break;
-						case '\r':
-							Writer.Write("\\r");
-							break;
-						case '\t':
-							Writer.Write("\\t");
-							break;
-						default:
-							if (Char.IsControl(Value[Idx]))
-							{
-								Writer.Write("\\u{0:X4}", (int)Value[Idx]);
-							}
-							else
-							{
-								Writer.Write(Value[Idx]);
-							}
-							break;
-					}
-				}
-			}
-			Writer.Write("\"");
-		}
-	}
 }
diff --git a/Engine/Source/Programs/DotNETCommon/DotNETUtilities/JsonWriter.cs b/Engine/Source/Programs/DotNETCommon/DotNETUtilities/JsonWriter.cs
new file mode 100644
--- /dev/null
+++ b/Engine/Source/Programs/DotNETCommon/DotNETUtilities/JsonWriter.cs
@@ -0,0 +1,248 @@
+using System;
+using System.Collections.Generic;
+using System.IO;
+using System.Linq;
+using System.Text;
+using System.Threading.Tasks;
+
+namespace Tools.DotNETCommon
+{
+	/// <summary>
+	/// Writer for JSON data, which indents the output text appropriately, and adds commas and newlines between fields
+	/// </summary>
+	public class JsonWriter : IDisposable
+	{
+		TextWriter Writer;
+		bool bLeaveOpen;
+		bool bRequiresComma;
+		string Indent;
+
+		/// <summary>
+		/// Constructor
+		/// </summary>
+		/// <param name="FileName">File to write to</param>
+		public JsonWriter(string FileName)
+			: this(new StreamWriter(FileName))
+		{
+		}
+
+		/// <summary>
+		/// Constructor
+		/// </summary>
+		/// <param name="FileName">File to write to</param>
+		public JsonWriter(FileReference FileName)
+			: this(new StreamWriter(FileName.FullName))
+		{
+		}
+
+		/// <summary>
+		/// Constructor
+		/// </summary>
+		/// <param name="Writer">The text writer to output to</param>
+		/// <param name="bLeaveOpen">Whether to leave the writer open when the object is disposed</param>
+		public JsonWriter(TextWriter Writer, bool bLeaveOpen = false)
+		{
+			this.Writer = Writer;
+			this.bLeaveOpen = bLeaveOpen;
+			Indent = "";
+		}
+
+		/// <summary>
+		/// Dispose of any managed resources
+		/// </summary>
+		public void Dispose()
+		{
+			if(!bLeaveOpen && Writer != null)
+			{
+				Writer.Dispose();
+				Writer = null;
+			}
+		}
+
+		/// <summary>
+		/// Write the opening brace for an object
+		/// </summary>
+		public void WriteObjectStart()
+		{
+			WriteCommaNewline();
+
+			Writer.Write(Indent);
+			Writer.Write("{");
+
+			Indent += "\t";
+			bRequiresComma = false;
+		}
+
+		/// <summary>
+		/// Write the name and opening brace for an object
+		/// </summary>
+		/// <param name="ObjectName">Name of the field</param>
+		public void WriteObjectStart(string ObjectName)
+		{
+			WriteCommaNewline();
+
+			Writer.Write("{0}\"{1}\": ", Indent, ObjectName);
+
+			bRequiresComma = false;
+
+			WriteObjectStart();
+		}
+
+		/// <summary>
+		/// Write the closing brace for an object
+		/// </summary>
+		public void WriteObjectEnd()
+		{
+			Indent = Indent.Substring(0, Indent.Length - 1);
+
+			Writer.WriteLine();
+			Writer.Write(Indent);
+			Writer.Write("}");
+
+			bRequiresComma = true;
+		}
+
+		/// <summary>
+		/// Write the name and opening bracket for an array
+		/// </summary>
+		/// <param name="ArrayName">Name of the field</param>
+		public void WriteArrayStart(string ArrayName)
+		{
+			WriteCommaNewline();
+
+			Writer.Write("{0}\"{1}\": [", Indent, ArrayName);
+
+			Indent += "\t";
+			bRequiresComma = false;
+		}
+
+		/// <summary>
+		/// Write the closing bracket for an array
+		/// </summary>
+		public void WriteArrayEnd()
+		{
+			Indent = Indent.Substring(0, Indent.Length - 1);
+
+			Writer.WriteLine();
+			Writer.Write("{0}]", Indent);
+
+			bRequiresComma = true;
+		}
+
+		/// <summary>
+		/// Write an array of strings
+		/// </summary>
+		/// <param name="Name">Name of the field</param>
+		/// <param name="Values">Values for the field</param>
+		public void WriteStringArrayField(string Name, IEnumerable<string> Values)
+		{
+			WriteArrayStart(Name);
+			foreach(string Value in Values)
+			{
+				WriteValue(Value);
+			}
+			WriteArrayEnd();
+		}
+
+		/// <summary>
+		/// Write an array of enum values
+		/// </summary>
+		/// <param name="Name">Name of the field</param>
+		/// <param name="Values">Values for the field</param>
+		public void WriteEnumArrayField<T>(string Name, IEnumerable<T> Values) where T : struct
+		{
+			WriteStringArrayField(Name, Values.Select(x => x.ToString()));
+		}
+
+		/// <summary>
+		/// Write a value with no field name, for the contents of an array
+		/// </summary>
+		/// <param name="Value">Value to write</param>
+		public void WriteValue(string Value)
+		{
+			WriteCommaNewline();
+
+			Writer.Write(Indent);
+			WriteEscapedString(Value);
+
+			bRequiresComma = true;
+		}
+
+		/// <summary>
+		/// Write a field name and string value
+		/// </summary>
+		/// <param name="Name">Name of the field</param>
+		/// <param name="Value">Value for the field</param>
+		public void WriteValue(string Name, string Value)
+		{
+			WriteCommaNewline();
+
+			Writer.Write("{0}\"{1}\": ", Indent, Name);
+			WriteEscapedString(Value);
+
+			bRequiresComma = true;
+		}
+
+		/// <summary>
+		/// Write a field name and integer value
+		/// </summary>
+		/// <param name="Name">Name of the field</param>
+		/// <param name="Value">Value for the field</param>
+		public void WriteValue(string Name, int Value)
+		{
+			WriteValueInternal(Name, Value.ToString());
+		}
+
+		/// <summary>
+		/// Write a field name and double value
+		/// </summary>
+		/// <param name="Name">Name of the field</param>
+		/// <param name="Value">Value for the field</param>
+		public void WriteValue(string Name, double Value)
+		{
+			WriteValueInternal(Name, Value.ToString());
+		}
+
+		/// <summary>
+		/// Write a field name and bool value
+		/// </summary>
+		/// <param name="Name">Name of the field</param>
+		/// <param name="Value">Value for the field</param>
+		public void WriteValue(string Name, bool Value)
+		{
+			WriteValueInternal(Name, Value ? "true" : "false");
+		}
+
+		void WriteCommaNewline()
+		{
+			if (bRequiresComma)
+			{
+				Writer.WriteLine(",");
+			}
+			else if (Indent.Length > 0)
+			{
+				Writer.WriteLine();
+			}
+		}
+
+		void WriteValueInternal(string Name, string Value)
+		{
+			WriteCommaNewline();
+
+			Writer.Write("{0}\"{1}\": {2}", Indent, Name, Value);
+
+			bRequiresComma = true;
+		}
+
+		void WriteEscapedString(string Value)
+		{
+			// Escape any characters which may not appear in a JSON string (see http://www.json.org).
+			Writer.Write("\"");
+			if (Value != null)
+			{
+				Writer.Write(Json.EscapeString(Value));
+			}
+			Writer.Write("\"");
+		}
+	}
+}
