Add a Json.Load() and Json.Save() function to quickly load and save an object to disk.

#rnx

diff --git a/Engine/Source/Programs/DotNETCommon/DotNETUtilities/Json.cs b/Engine/Source/Programs/DotNETCommon/DotNETUtilities/Json.cs
--- a/Engine/Source/Programs/DotNETCommon/DotNETUtilities/Json.cs
+++ b/Engine/Source/Programs/DotNETCommon/DotNETUtilities/Json.cs
@@ -77,6 +77,45 @@ namespace Tools.DotNETCommon
 			return fastJSON.JSON.Instance.ToObject<T>(Text);
 		}
 
+		/// <summary>
+		/// Read a JSON object from a file on disk
+		/// </summary>
+		/// <typeparam name="T">The type of the object to load</typeparam>
+		/// <param name="Location">The location of the file to read</param>
+		/// <returns>The deserialized object</returns>
+		public static T Load<T>(FileReference Location)
+		{
+			string Text;
+			try
+			{
+				Text = FileReference.ReadAllText(Location);
+			}
+			catch(Exception Ex)
+			{
+				throw new Exception(String.Format("Unable to read '{0}'", Location), Ex);
+			}
+			return Deserialize<T>(Text);
+		}
+
+		/// <summary>
+		/// Save a JSON object to a file on disk
+		/// </summary>
+		/// <param name="Location">The location of the file to read</param>
+		/// <param name="Object">The object to save</param>
+		public static void Save(FileReference Location, object Object)
+		{
+			string Text = Serialize(Object);
+			try
+			{
+				DirectoryReference.CreateDirectory(Location.Directory);
+				FileReference.WriteAllText(Location, Text);
+			}
+			catch(Exception Ex)
+			{
+				throw new Exception(String.Format("Unable to write '{0}'", Location), Ex);
+			}
+		}
+
 		/// <summary>
 		/// Formats the given JSON string in 'Epic' format (parentheses/brackets on new lines, tabs instead of spaces, no space before colons in key/value pairs)
 		/// </summary>
