Remove the StringUtils.FastEndsWith() method. The reason this is slow is that it uses a culture-aware comparison. Using the ordinal comparer from the FileSystemReference object should achieve the same result.

#rnx

diff --git a/Engine/Source/Programs/DotNETCommon/DotNETUtilities/StringUtils.cs b/Engine/Source/Programs/DotNETCommon/DotNETUtilities/StringUtils.cs
--- a/Engine/Source/Programs/DotNETCommon/DotNETUtilities/StringUtils.cs
+++ b/Engine/Source/Programs/DotNETCommon/DotNETUtilities/StringUtils.cs
@@ -21,31 +21,6 @@ namespace Tools.DotNETCommon
 			Builder.AppendLine();
 		}
 
-		/// <summary>
-		/// Case-sensitive replacement for String.EndsWith(), which seems to be pathalogically slow on Mono.
-		/// </summary>
-		/// <param name="Source">String to test</param>
-		/// <param name="Suffix">Suffix to check for</param>
-		/// <returns>True if the source string ends with the given suffix</returns>
-		public static bool FastEndsWith(string Source, string Suffix)
-		{
-			if(Source.Length < Suffix.Length)
-			{
-				return false;
-			}
-
-			int SourceBase = Source.Length - Suffix.Length;
-			for(int Idx = 0; Idx < Suffix.Length; Idx++)
-			{
-				if(Source[SourceBase + Idx] != Suffix[Idx])
-				{
-					return false;
-				}
-			}
-
-			return true;
-		}
-
 		/// <summary>
 		/// Formats a list of strings in the style "1, 2, 3 and 4"
 		/// </summary>
diff --git a/Engine/Source/Programs/UnrealBuildTool/System/SourceFileSearch.cs b/Engine/Source/Programs/UnrealBuildTool/System/SourceFileSearch.cs
--- a/Engine/Source/Programs/UnrealBuildTool/System/SourceFileSearch.cs
+++ b/Engine/Source/Programs/UnrealBuildTool/System/SourceFileSearch.cs
@@ -106,7 +106,7 @@ namespace UnrealBuildTool
 			}
 		}
 
-		static bool ShouldInclude(FileSystemReference Reference, string[] InvariantExcludedSuffixes)
+		static bool ShouldInclude(FileSystemReference Reference, string[] ExcludedSuffixes)
 		{
 			// Ignore Mac resource fork files on non-HFS partitions
 			if (Path.GetFileName(Reference.FullName).StartsWith("._"))
@@ -114,9 +114,9 @@ namespace UnrealBuildTool
 				return false;
 			}
 			
-			foreach (string InvariantExcludedSuffix in InvariantExcludedSuffixes)
+			foreach (string ExcludedSuffix in ExcludedSuffixes)
 			{
-				if (StringUtils.FastEndsWith(Reference.FullName, InvariantExcludedSuffix))
+				if (Reference.FullName.EndsWith(ExcludedSuffix, FileSystemReference.Comparison))
 				{
 					return false;
 				}
