Wrap setting process priority in a try/catch in case the process exists before the priority is able to be set

diff --git a/Engine/Source/Programs/DotNETCommon/DotNETUtilities/ManagedProcess.cs b/Engine/Source/Programs/DotNETCommon/DotNETUtilities/ManagedProcess.cs
--- a/Engine/Source/Programs/DotNETCommon/DotNETUtilities/ManagedProcess.cs
+++ b/Engine/Source/Programs/DotNETCommon/DotNETUtilities/ManagedProcess.cs
@@ -483,10 +483,16 @@ namespace Tools.DotNETCommon
 				}
 
 				FrameworkProcess.Start();
-				FrameworkProcess.PriorityClass = Priority;
+				try
+				{
+					FrameworkProcess.PriorityClass = Priority;
+				}
+				catch
+				{
 
+				}
 				// If we have any input text, write it to stdin now
-				if(Input != null)
+				if (Input != null)
 				{
 					FrameworkProcess.StandardInput.BaseStream.Write(Input, 0, Input.Length);
 					FrameworkProcess.StandardInput.BaseStream.Close();
