Add a custom binary reader/writer class to DotNETUtilities which supports object graphs, containers, and nullable objects.

#rnx

diff --git a/Engine/Source/Programs/DotNETCommon/DotNETUtilities/BinaryArchiveReader.cs b/Engine/Source/Programs/DotNETCommon/DotNETUtilities/BinaryArchiveReader.cs
new file mode 100644
--- /dev/null
+++ b/Engine/Source/Programs/DotNETCommon/DotNETUtilities/BinaryArchiveReader.cs
@@ -0,0 +1,449 @@
+using System;
+using System.Collections.Generic;
+using System.IO;
+using System.Linq;
+using System.Text;
+using System.Threading.Tasks;
+
+namespace Tools.DotNETCommon
+{
+	/// <summary>
+	/// Reads data from a binary output stream. Similar to the NET Framework BinaryReader class, but supports fast serialization of object graphs and container types, and supports nullable objects.
+	/// Significantly faster than BinaryReader due to the expectation that the whole stream is in memory before deserialization.
+	/// </summary>
+	public class BinaryArchiveReader : IDisposable
+	{
+		/// <summary>
+		/// The input stream.
+		/// </summary>
+		Stream Stream;
+
+		/// <summary>
+		/// The input buffer
+		/// </summary>
+		byte[] Buffer;
+
+		/// <summary>
+		/// Current position within the buffer
+		/// </summary>
+		int BufferPos;
+		
+		/// <summary>
+		/// List of previously serialized objects
+		/// </summary>
+		List<object> Objects = new List<object>();
+
+		/// <summary>
+		/// Constructor
+		/// </summary>
+		/// <param name="Buffer">The buffer to read from</param>
+		public BinaryArchiveReader(byte[] Buffer)
+		{
+			this.Buffer = Buffer;
+		}
+
+		/// <summary>
+		/// Constructor
+		/// </summary>
+		/// <param name="FileName">File to read from</param>
+		public BinaryArchiveReader(FileReference FileName)
+		{
+			this.Buffer = FileReference.ReadAllBytes(FileName);
+		}
+
+		/// <summary>
+		/// Constructor
+		/// </summary>
+		/// <param name="Stream">Stream to read from</param>
+		public BinaryArchiveReader(Stream Stream)
+		{
+			Buffer = new byte[Stream.Length];
+			Stream.Read(Buffer, 0, Buffer.Length);
+		}
+
+		/// <summary>
+		/// Dispose of the stream owned by this reader
+		/// </summary>
+		public void Dispose()
+		{
+			if(Stream != null)
+			{
+				Stream.Dispose();
+				Stream = null;
+			}
+
+			Buffer = null;
+		}
+
+		/// <summary>
+		/// Reads a bool from the stream
+		/// </summary>
+		/// <returns>The value that was read</returns>
+		public bool ReadBool()
+		{
+			return ReadByte() != 0;
+		}
+
+		/// <summary>
+		/// Reads a single byte from the stream
+		/// </summary>
+		/// <returns>The value that was read</returns>
+		public byte ReadByte()
+		{
+			byte Value = Buffer[BufferPos];
+			BufferPos++;
+			return Value;
+		}
+
+		/// <summary>
+		/// Reads a single signed byte from the stream
+		/// </summary>
+		/// <returns>The value that was read</returns>
+		public sbyte ReadSignedByte()
+		{
+			return (sbyte)ReadByte();
+		}
+
+		/// <summary>
+		/// Reads a single short from the stream
+		/// </summary>
+		/// <returns>The value that was read</returns>
+		public short ReadShort()
+		{
+			return (short)ReadUnsignedShort();
+		}
+
+		/// <summary>
+		/// Reads a single unsigned short from the stream
+		/// </summary>
+		/// <returns>The value that was read</returns>
+		public ushort ReadUnsignedShort()
+		{
+			ushort Value = (ushort)(Buffer[BufferPos + 0] | (Buffer[BufferPos + 1] << 8));
+			BufferPos += 2;
+			return Value;
+		}
+
+		/// <summary>
+		/// Reads a single int from the stream
+		/// </summary>
+		/// <returns>The value that was read</returns>
+		public int ReadInt()
+		{
+			return (int)ReadUnsignedInt();
+		}
+
+		/// <summary>
+		/// Reads a single unsigned int from the stream
+		/// </summary>
+		/// <returns>The value that was read</returns>
+		public uint ReadUnsignedInt()
+		{
+			uint Value = (uint)(Buffer[BufferPos + 0] | (Buffer[BufferPos + 1] << 8) | (Buffer[BufferPos + 2] << 16) | (Buffer[BufferPos + 3] << 24));
+			BufferPos += 4;
+			return Value;
+		}
+
+		/// <summary>
+		/// Reads a single long from the stream
+		/// </summary>
+		/// <returns>The value that was read</returns>
+		public long ReadLong()
+		{
+			return (long)ReadUnsignedLong();
+		}
+
+		/// <summary>
+		/// Reads a single unsigned long from the stream
+		/// </summary>
+		/// <returns>The value that was read</returns>
+		public ulong ReadUnsignedLong()
+		{
+			ulong Value = (ulong)ReadUnsignedInt();
+			Value |= (ulong)ReadUnsignedInt() << 32;
+			return Value;
+		}
+
+		/// <summary>
+		/// Reads a string from the stream
+		/// </summary>
+		/// <returns>The value that was read</returns>
+		public string ReadString()
+		{
+			byte[] Bytes = ReadByteArray();
+			if(Bytes == null)
+			{
+				return null;
+			}
+			else
+			{
+				return Encoding.UTF8.GetString(Bytes);
+			}
+		}
+
+		/// <summary>
+		/// Reads a byte array from the stream
+		/// </summary>
+		/// <returns>The data that was read</returns>
+		public byte[] ReadByteArray()
+		{
+			return ReadPrimitiveArray<byte>(sizeof(byte));
+		}
+
+		/// <summary>
+		/// Reads a short array from the stream
+		/// </summary>
+		/// <returns>The data that was read</returns>
+		public short[] ReadShortArray()
+		{
+			return ReadPrimitiveArray<short>(sizeof(short));
+		}
+
+		/// <summary>
+		/// Reads an int array from the stream
+		/// </summary>
+		/// <returns>The data that was read</returns>
+		public int[] ReadIntArray()
+		{
+			return ReadPrimitiveArray<int>(sizeof(int));
+		}
+
+		/// <summary>
+		/// Reads an array of primitive types from the stream
+		/// </summary>
+		/// <param name="ElementSize">Size of a single element</param>
+		/// <returns>The data that was read</returns>
+		private T[] ReadPrimitiveArray<T>(int ElementSize) where T : struct
+		{
+			int Length = ReadInt();
+			if(Length < 0)
+			{
+				return null;
+			}
+			else
+			{
+				T[] Result = new T[Length];
+				ReadBulkData(Result, Length * ElementSize);
+				return Result;
+			}
+		}
+
+		/// <summary>
+		/// Reads bulk data from the stream into the given buffer
+		/// </summary>
+		/// <param name="Data">Array which receives the data that was read</param>
+		/// <param name="Size">Size of data to read</param>
+		private void ReadBulkData(Array Data, int Size)
+		{
+			System.Buffer.BlockCopy(Buffer, BufferPos, Data, 0, Size);
+			BufferPos += Size;
+		}
+
+
+		/// <summary>
+		/// Reads an array of items
+		/// </summary>
+		/// <returns>New array</returns>
+		public T[] ReadArray<T>(Func<T> ReadElement)
+		{
+			int Count = ReadInt();
+			if(Count < 0)
+			{
+				return null;
+			}
+			else
+			{
+				T[] Result = new T[Count];
+				for(int Idx = 0; Idx < Count; Idx++)
+				{
+					Result[Idx] = ReadElement();
+				}
+				return Result;
+			}
+		}
+
+		/// <summary>
+		/// Reads a list of items
+		/// </summary>
+		/// <typeparam name="T">The element type for the list</typeparam>
+		/// <param name="ReadElement">Delegate used to read a single element</param>
+		/// <returns>List of items</returns>
+		public List<T> ReadList<T>(Func<T> ReadElement)
+		{
+			int Count = ReadInt();
+			if(Count < 0)
+			{
+				return null;
+			}
+			else
+			{
+				List<T> Result = new List<T>(Count);
+				for(int Idx = 0; Idx < Count; Idx++)
+				{
+					Result.Add(ReadElement());
+				}
+				return Result;
+			}
+		}
+
+		/// <summary>
+		/// Reads a hashset of items
+		/// </summary>
+		/// <typeparam name="T">The element type for the set</typeparam>
+		/// <param name="ReadElement">Delegate used to read a single element</param>
+		/// <returns>Set of items</returns>
+		public HashSet<T> ReadHashSet<T>(Func<T> ReadElement)
+		{
+			int Count = ReadInt();
+			if(Count < 0)
+			{
+				return null;
+			}
+			else
+			{
+				HashSet<T> Result = new HashSet<T>();
+				for(int Idx = 0; Idx < Count; Idx++)
+				{
+					Result.Add(ReadElement());
+				}
+				return Result;
+			}
+		}
+
+		/// <summary>
+		/// Reads a hashset of items
+		/// </summary>
+		/// <typeparam name="T">The element type for the set</typeparam>
+		/// <param name="ReadElement">Delegate used to read a single element</param>
+		/// <param name="Comparer">Comparison function for the set</param>
+		/// <returns>Set of items</returns>
+		public HashSet<T> ReadHashSet<T>(Func<T> ReadElement, IEqualityComparer<T> Comparer)
+		{
+			int Count = ReadInt();
+			if(Count < 0)
+			{
+				return null;
+			}
+			else
+			{
+				HashSet<T> Result = new HashSet<T>(Comparer);
+				for(int Idx = 0; Idx < Count; Idx++)
+				{
+					Result.Add(ReadElement());
+				}
+				return Result;
+			}
+		}
+
+		/// <summary>
+		/// Reads a dictionary of items
+		/// </summary>
+		/// <typeparam name="K">Type of the dictionary key</typeparam>
+		/// <typeparam name="V">Type of the dictionary value</typeparam>
+		/// <param name="ReadKey">Delegate used to read a single key</param>
+		/// <param name="ReadValue">Delegate used to read a single value</param>
+		/// <returns>New dictionary instance</returns>
+		public Dictionary<K, V> ReadDictionary<K, V>(Func<K> ReadKey, Func<V> ReadValue)
+		{
+			int Count = ReadInt();
+			if(Count < 0)
+			{
+				return null;
+			}
+			else
+			{
+				Dictionary<K, V> Result = new Dictionary<K, V>(Count);
+				for(int Idx = 0; Idx < Count; Idx++)
+				{
+					Result.Add(ReadKey(), ReadValue());
+				}
+				return Result;
+			}
+		}
+
+		/// <summary>
+		/// Reads a dictionary of items
+		/// </summary>
+		/// <typeparam name="K">Type of the dictionary key</typeparam>
+		/// <typeparam name="V">Type of the dictionary value</typeparam>
+		/// <param name="ReadKey">Delegate used to read a single key</param>
+		/// <param name="ReadValue">Delegate used to read a single value</param>
+		/// <param name="Comparer">Comparison function for keys in the dictionary</param>
+		/// <returns>New dictionary instance</returns>
+		public Dictionary<K, V> ReadDictionary<K, V>(Func<K> ReadKey, Func<V> ReadValue, IEqualityComparer<K> Comparer)
+		{
+			int Count = ReadInt();
+			if(Count < 0)
+			{
+				return null;
+			}
+			else
+			{
+				Dictionary<K, V> Result = new Dictionary<K, V>(Count, Comparer);
+				for(int Idx = 0; Idx < Count; Idx++)
+				{
+					Result.Add(ReadKey(), ReadValue());
+				}
+				return Result;
+			}
+		}
+
+		/// <summary>
+		/// Reads an object reference from the stream. Each referenced object will only be serialized once using the supplied delegates. Reading an object instance is
+		/// done in two phases; first the object is created and its reference stored in the unique object list, then the object contents are read. This allows the object to 
+		/// serialize a reference to itself.
+		/// </summary>
+		/// <typeparam name="T">Type of the object to read.</typeparam>
+		/// <param name="CreateObject">Delegate used to create an object instance</param>
+		/// <param name="ReadObject">Delegate used to read an object instance</param>
+		/// <returns>Object instance</returns>
+		public T ReadObjectReference<T>(Func<T> CreateObject, Action<T> ReadObject) where T : class
+		{
+			int Index = ReadInt();
+			if(Index < 0)
+			{
+				return null;
+			}
+			else
+			{
+				if(Index == Objects.Count)
+				{
+					T Object = CreateObject();
+					Objects.Add(Object);
+					ReadObject(Object);
+				}
+				return (T)Objects[Index];
+			}
+		}
+
+		/// <summary>
+		/// Reads an object reference from the stream. Each object will only be serialized once using the supplied delegate; subsequent reads reference the original.
+		/// Since the reader only receives the object reference when the CreateObject delegate returns, it is not possible for the object to serialize a reference to itself.
+		/// </summary>
+		/// <typeparam name="T">Type of the object to read.</typeparam>
+		/// <param name="ReadObject">Delegate used to create an object instance. The object may not reference itself recursively.</param>
+		/// <returns>Object instance</returns>
+		public T ReadObjectReference<T>(Func<T> ReadObject) where T : class
+		{
+			int Index = ReadInt();
+			if(Index < 0)
+			{
+				return null;
+			}
+			else
+			{
+				// Temporarily add the reader to the object list, so we can detect invalid recursive references. 
+				if(Index == Objects.Count)
+				{
+					Objects.Add(null);
+					Objects[Index] = ReadObject();
+				}
+				if(Objects[Index] == null)
+				{
+					throw new InvalidOperationException("Attempt to serialize reference to object recursively.");
+				}
+				return (T)Objects[Index];
+			}
+		}
+	}
+}
diff --git a/Engine/Source/Programs/DotNETCommon/DotNETUtilities/BinaryArchiveWriter.cs b/Engine/Source/Programs/DotNETCommon/DotNETUtilities/BinaryArchiveWriter.cs
new file mode 100644
--- /dev/null
+++ b/Engine/Source/Programs/DotNETCommon/DotNETUtilities/BinaryArchiveWriter.cs
@@ -0,0 +1,434 @@
+using System;
+using System.Collections.Generic;
+using System.IO;
+using System.Linq;
+using System.Runtime.CompilerServices;
+using System.Text;
+using System.Threading.Tasks;
+
+namespace Tools.DotNETCommon
+{
+	/// <summary>
+	/// Writes data to a binary output stream. Similar to the NET Framework BinaryWriter class, but supports fast serialization of object graphs and container types, and supports nullable objects.
+	/// </summary>
+	public class BinaryArchiveWriter : IDisposable
+	{
+		/// <summary>
+		/// Comparer which tests for reference equality between two objects
+		/// </summary>
+		class ReferenceComparer : IEqualityComparer<object>
+		{
+			bool IEqualityComparer<object>.Equals(object A, object B)
+			{
+				return A == B;
+			}
+
+			int IEqualityComparer<object>.GetHashCode(object X)
+			{
+				return RuntimeHelpers.GetHashCode(X);
+			}
+		}
+
+		/// <summary>
+		/// Instance of the ReferenceComparer class which can be shared by all archive writers
+		/// </summary>
+		static ReferenceComparer ReferenceComparerInstance = new ReferenceComparer();
+
+		/// <summary>
+		/// The output stream being written to
+		/// </summary>
+		Stream Stream;
+
+		/// <summary>
+		/// Buffer for data to be written to the stream
+		/// </summary>
+		byte[] Buffer;
+
+		/// <summary>
+		/// Current position within the output buffer
+		/// </summary>
+		int BufferPos;
+
+		/// <summary>
+		/// Map of object instance to unique id
+		/// </summary>
+		Dictionary<object, int> ObjectToUniqueId = new Dictionary<object, int>(ReferenceComparerInstance);
+
+		/// <summary>
+		/// Constructor
+		/// </summary>
+		/// <param name="Stream">The output stream</param>
+		public BinaryArchiveWriter(Stream Stream)
+		{
+			this.Stream = Stream;
+			this.Buffer = new byte[4096];
+		}
+
+		/// <summary>
+		/// Constructor
+		/// </summary>
+		/// <param name="FileName">File to write to</param>
+		public BinaryArchiveWriter(FileReference FileName)
+			: this(File.Open(FileName.FullName, FileMode.Create, FileAccess.Write, FileShare.Read))
+		{
+		}
+
+		/// <summary>
+		/// Flushes this stream, and disposes the stream
+		/// </summary>
+		public void Dispose()
+		{
+			Flush();
+
+			if(Stream != null)
+			{
+				Stream.Dispose();
+				Stream = null;
+			}
+		}
+
+		/// <summary>
+		/// Writes all buffered data to disk
+		/// </summary>
+		public void Flush()
+		{
+			if(BufferPos > 0)
+			{
+				Stream.Write(Buffer, 0, BufferPos);
+				BufferPos = 0;
+			}
+		}
+
+		/// <summary>
+		/// Ensures there is a minimum amount of space in the output buffer
+		/// </summary>
+		/// <param name="NumBytes">Minimum amount of space required in the output buffer</param>
+		private void EnsureSpace(int NumBytes)
+		{
+			if(BufferPos + NumBytes > Buffer.Length)
+			{
+				Flush();
+				if(NumBytes > Buffer.Length)
+				{
+					Buffer = new byte[NumBytes];
+				}
+			}
+		}
+
+		/// <summary>
+		/// Writes a bool to the output
+		/// </summary>
+		/// <param name="Value">Value to write</param>
+		public void WriteBool(bool Value)
+		{
+			WriteByte(Value? (byte)1 : (byte)0);
+		}
+
+		/// <summary>
+		/// Writes a single byte to the output
+		/// </summary>
+		/// <param name="Value">Value to write</param>
+		public void WriteByte(byte Value)
+		{
+			EnsureSpace(1);
+
+			Buffer[BufferPos] = Value;
+
+			BufferPos++;
+		}
+
+		/// <summary>
+		/// Writes a single signed byte to the output
+		/// </summary>
+		/// <param name="Value">Value to write</param>
+		public void WriteSignedByte(sbyte Value)
+		{
+			WriteByte((byte)Value);
+		}
+
+		/// <summary>
+		/// Writes a single short to the output
+		/// </summary>
+		/// <param name="Value">Value to write</param>
+		public void WriteShort(short Value)
+		{
+			WriteUnsignedShort((ushort)Value);
+		}
+
+		/// <summary>
+		/// Writes a single unsigned short to the output
+		/// </summary>
+		/// <param name="Value">Value to write</param>
+		public void WriteUnsignedShort(ushort Value)
+		{
+			EnsureSpace(2);
+
+			Buffer[BufferPos + 0] = (byte)Value;
+			Buffer[BufferPos + 1] = (byte)(Value >> 8);
+
+			BufferPos += 2;
+		}
+
+		/// <summary>
+		/// Writes a single int to the output
+		/// </summary>
+		/// <param name="Value">Value to write</param>
+		public void WriteInt(int Value)
+		{
+			WriteUnsignedInt((uint)Value);
+		}
+
+		/// <summary>
+		/// Writes a single unsigned int to the output
+		/// </summary>
+		/// <param name="Value">Value to write</param>
+		public void WriteUnsignedInt(uint Value)
+		{
+			EnsureSpace(4);
+
+			Buffer[BufferPos + 0] = (byte)Value;
+			Buffer[BufferPos + 1] = (byte)(Value >> 8);
+			Buffer[BufferPos + 2] = (byte)(Value >> 16);
+			Buffer[BufferPos + 3] = (byte)(Value >> 24);
+
+			BufferPos += 4;
+		}
+
+		/// <summary>
+		/// Writes a single long to the output
+		/// </summary>
+		/// <param name="Value">Value to write</param>
+		public void WriteLong(long Value)
+		{
+			WriteUnsignedLong((ulong)Value);
+		}
+
+		/// <summary>
+		/// Writes a single unsigned long to the output
+		/// </summary>
+		/// <param name="Value">Value to write</param>
+		public void WriteUnsignedLong(ulong Value)
+		{
+			EnsureSpace(8);
+
+			Buffer[BufferPos + 0] = (byte)Value;
+			Buffer[BufferPos + 1] = (byte)(Value >> 8);
+			Buffer[BufferPos + 2] = (byte)(Value >> 16);
+			Buffer[BufferPos + 3] = (byte)(Value >> 24);
+			Buffer[BufferPos + 4] = (byte)(Value >> 32);
+			Buffer[BufferPos + 5] = (byte)(Value >> 40);
+			Buffer[BufferPos + 6] = (byte)(Value >> 48);
+			Buffer[BufferPos + 7] = (byte)(Value >> 56);
+
+			BufferPos += 8;
+		}
+
+		/// <summary>
+		/// Writes a string to the output
+		/// </summary>
+		/// <param name="Value">Value to write</param>
+		public void WriteString(string Value)
+		{
+			byte[] Bytes;
+			if(Value == null)
+			{
+				Bytes = null;
+			}
+			else
+			{
+				Bytes = Encoding.UTF8.GetBytes(Value);
+			}
+			WriteByteArray(Bytes);
+		}
+
+		/// <summary>
+		/// Writes an array of bytes to the output
+		/// </summary>
+		/// <param name="Data">Data to write. May be null.</param>
+		public void WriteByteArray(byte[] Data)
+		{
+			WritePrimitiveArray(Data, sizeof(byte));
+		}
+
+		/// <summary>
+		/// Writes an array of shorts to the output
+		/// </summary>
+		/// <param name="Data">Data to write. May be null.</param>
+		public void WriteShortArray(short[] Data)
+		{
+			WritePrimitiveArray(Data, sizeof(short));
+		}
+
+		/// <summary>
+		/// Writes an array of ints to the output
+		/// </summary>
+		/// <param name="Data">Data to write. May be null.</param>
+		public void WriteIntArray(int[] Data)
+		{
+			WritePrimitiveArray(Data, sizeof(int));
+		}
+
+		/// <summary>
+		/// Writes an array of primitive types to the output.
+		/// </summary>
+		/// <param name="Data">Data to write. May be null.</param>
+		private void WritePrimitiveArray<T>(T[] Data, int ElementSize) where T : struct
+		{
+			if(Data == null)
+			{
+				WriteInt(-1);
+			}
+			else
+			{
+				WriteInt(Data.Length);
+				WriteBulkData(Data, Data.Length * ElementSize);
+			}
+		}
+
+		/// <summary>
+		/// Writes primitive data from the given array to the output buffer.
+		/// </summary>
+		/// <param name="Data">Data to write.</param>
+		/// <param name="Size">Size of the data, in bytes</param>
+		private void WriteBulkData(Array Data, int Size)
+		{
+			if(Size > 0)
+			{
+				for(int Pos = 0;; )
+				{
+					int CopySize = Math.Min(Size - Pos, Buffer.Length - BufferPos);
+
+					System.Buffer.BlockCopy(Data, Pos, Buffer, BufferPos, CopySize);
+					BufferPos += CopySize;
+					Pos += CopySize;
+
+					if(Pos == Size)
+					{
+						break;
+					}
+
+					Flush();
+				}
+			}
+		}
+
+		/// <summary>
+		/// Write an array of items to the archive
+		/// </summary>
+		/// <typeparam name="T">Type of the element</typeparam>
+		/// <param name="Items">Array of items</param>
+		/// <param name="WriteElement">Writes an individual element to the archive</param>
+		public void WriteArray<T>(T[] Items, Action<T> WriteElement)
+		{
+			if(Items == null)
+			{
+				WriteInt(-1);
+			}
+			else
+			{
+				WriteInt(Items.Length);
+				for(int Idx = 0; Idx < Items.Length; Idx++)
+				{
+					WriteElement(Items[Idx]);
+				}
+			}
+		}
+
+		/// <summary>
+		/// Write a list of items to the archive
+		/// </summary>
+		/// <typeparam name="T">Type of the element</typeparam>
+		/// <param name="Items">List of items</param>
+		/// <param name="WriteElement">Writes an individual element to the archive</param>
+		public void WriteList<T>(List<T> Items, Action<T> WriteElement)
+		{
+			if(Items == null)
+			{
+				WriteInt(-1);
+			}
+			else
+			{
+				WriteInt(Items.Count);
+				for(int Idx = 0; Idx < Items.Count; Idx++)
+				{
+					WriteElement(Items[Idx]);
+				}
+			}
+		}
+
+		/// <summary>
+		/// Writes a hashset of items
+		/// </summary>
+		/// <typeparam name="T">The element type for the set</typeparam>
+		/// <param name="Set">The set to write</param>
+		/// <param name="WriteElement">Delegate used to read a single element</param>
+		public void WriteHashSet<T>(HashSet<T> Set, Action<T> WriteElement)
+		{
+			if(Set == null)
+			{
+				WriteInt(-1);
+			}
+			else
+			{
+				WriteInt(Set.Count);
+				foreach(T Element in Set)
+				{
+					WriteElement(Element);
+				}
+			}
+		}
+
+		/// <summary>
+		/// Writes a dictionary of items
+		/// </summary>
+		/// <typeparam name="K">Type of the dictionary key</typeparam>
+		/// <typeparam name="V">Type of the dictionary value</typeparam>
+		/// <param name="Dictionary">The dictionary to write</param>
+		/// <param name="WriteKey">Delegate used to read a single key</param>
+		/// <param name="WriteValue">Delegate used to read a single value</param>
+		public void WriteDictionary<K, V>(Dictionary<K, V> Dictionary, Action<K> WriteKey, Action<V> WriteValue)
+		{
+			if(Dictionary == null)
+			{
+				WriteInt(-1);
+			}
+			else
+			{
+				WriteInt(Dictionary.Count);
+				foreach(KeyValuePair<K, V> Pair in Dictionary)
+				{
+					WriteKey(Pair.Key);
+					WriteValue(Pair.Value);
+				}
+			}
+		}
+
+		/// <summary>
+		/// Writes an object to the output. If the specific instance has already been written, preserves the reference to that.
+		/// </summary>
+		/// <typeparam name="T">Type of the object to serialize</typeparam>
+		/// <param name="Object">The object to serialize</param>
+		/// <param name="WriteObject">Delegate used to write the object</param>
+		public void WriteObjectReference<T>(T Object, Action WriteObject) where T : class
+		{
+			if(Object == null)
+			{
+				WriteInt(-1);
+			}
+			else
+			{
+				int Index;
+				if(ObjectToUniqueId.TryGetValue(Object, out Index))
+				{
+					WriteInt(Index);
+				}
+				else
+				{
+					WriteInt(ObjectToUniqueId.Count);
+					ObjectToUniqueId.Add(Object, ObjectToUniqueId.Count);
+					WriteObject();
+				}
+			}
+		}
+	}
+}
diff --git a/Engine/Source/Programs/DotNETCommon/DotNETUtilities/DirectoryReference.cs b/Engine/Source/Programs/DotNETCommon/DotNETUtilities/DirectoryReference.cs
--- a/Engine/Source/Programs/DotNETCommon/DotNETUtilities/DirectoryReference.cs
+++ b/Engine/Source/Programs/DotNETCommon/DotNETUtilities/DirectoryReference.cs
@@ -428,5 +428,40 @@ namespace Tools.DotNETCommon
 			string FullName = Reader.ReadString();
 			return (FullName.Length == 0) ? null : new DirectoryReference(FullName, DirectoryReference.Sanitize.None);
 		}
+
+		/// <summary>
+		/// Writes a directory reference  to a binary archive
+		/// </summary>
+		/// <param name="Writer">The writer to output data to</param>
+		/// <param name="Directory">The item to write</param>
+		public static void WriteDirectoryReference(this BinaryArchiveWriter Writer, DirectoryReference Directory)
+		{
+			if(Directory == null)
+			{
+				Writer.WriteString(null);
+			}
+			else
+			{
+				Writer.WriteString(Directory.FullName);
+			}
+		}
+
+		/// <summary>
+		/// Reads a directory reference from a binary archive
+		/// </summary>
+		/// <param name="Reader">Reader to serialize data from</param>
+		/// <returns>New directory reference instance</returns>
+		public static DirectoryReference ReadDirectoryReference(this BinaryArchiveReader Reader)
+		{
+			string FullName = Reader.ReadString();
+			if(FullName == null)
+			{
+				return null;
+			}
+			else
+			{
+				return new DirectoryReference(FullName, DirectoryReference.Sanitize.None);
+			}
+		}
 	}
 }
diff --git a/Engine/Source/Programs/DotNETCommon/DotNETUtilities/DotNETUtilities.csproj b/Engine/Source/Programs/DotNETCommon/DotNETUtilities/DotNETUtilities.csproj
--- a/Engine/Source/Programs/DotNETCommon/DotNETUtilities/DotNETUtilities.csproj
+++ b/Engine/Source/Programs/DotNETCommon/DotNETUtilities/DotNETUtilities.csproj
@@ -57,6 +57,8 @@
     <Compile Include="AssemblyUtils.cs" />
     <Compile Include="BinaryFormatterUtils.cs" />
     <Compile Include="BinaryReaderExtensions.cs" />
+    <Compile Include="BinaryArchiveReader.cs" />
+    <Compile Include="BinaryArchiveWriter.cs" />
     <Compile Include="BinaryWriterExtensions.cs" />
     <Compile Include="CommandLineArguments.cs" />
     <Compile Include="CommandLineAttribute.cs" />
diff --git a/Engine/Source/Programs/DotNETCommon/DotNETUtilities/FileReference.cs b/Engine/Source/Programs/DotNETCommon/DotNETUtilities/FileReference.cs
--- a/Engine/Source/Programs/DotNETCommon/DotNETUtilities/FileReference.cs
+++ b/Engine/Source/Programs/DotNETCommon/DotNETUtilities/FileReference.cs
@@ -604,5 +604,40 @@ namespace Tools.DotNETCommon
 				return Result;
 			}
 		}
+
+		/// <summary>
+		/// Writes a FileReference to a binary archive
+		/// </summary>
+		/// <param name="Writer">The writer to output data to</param>
+		/// <param name="File">The file reference to write</param>
+		public static void WriteFileReference(this BinaryArchiveWriter Writer, FileReference File)
+		{
+			if(File == null)
+			{
+				Writer.WriteString(null);
+			}
+			else
+			{
+				Writer.WriteString(File.FullName);
+			}
+		}
+
+		/// <summary>
+		/// Reads a FileReference from a binary archive
+		/// </summary>
+		/// <param name="Reader">Reader to serialize data from</param>
+		/// <returns>New file reference instance</returns>
+		public static FileReference ReadFileReference(this BinaryArchiveReader Reader)
+		{
+			string FullName = Reader.ReadString();
+			if(FullName == null)
+			{
+				return null;
+			}
+			else
+			{
+				return new FileReference(FullName, FileReference.Sanitize.None);
+			}
+		}
 	}
 }
