Add missing copyright notices.

#rnx

diff --git a/Engine/Source/Programs/DotNETCommon/DotNETUtilities/BinaryArchiveReader.cs b/Engine/Source/Programs/DotNETCommon/DotNETUtilities/BinaryArchiveReader.cs
--- a/Engine/Source/Programs/DotNETCommon/DotNETUtilities/BinaryArchiveReader.cs
+++ b/Engine/Source/Programs/DotNETCommon/DotNETUtilities/BinaryArchiveReader.cs
@@ -1,3 +1,5 @@
+// Copyright 1998-2019 Epic Games, Inc. All Rights Reserved.
+
 using System;
 using System.Collections.Generic;
 using System.IO;
diff --git a/Engine/Source/Programs/DotNETCommon/DotNETUtilities/BinaryArchiveWriter.cs b/Engine/Source/Programs/DotNETCommon/DotNETUtilities/BinaryArchiveWriter.cs
--- a/Engine/Source/Programs/DotNETCommon/DotNETUtilities/BinaryArchiveWriter.cs
+++ b/Engine/Source/Programs/DotNETCommon/DotNETUtilities/BinaryArchiveWriter.cs
@@ -1,3 +1,5 @@
+// Copyright 1998-2019 Epic Games, Inc. All Rights Reserved.
+
 using System;
 using System.Collections.Generic;
 using System.IO;
