Fix race condition in ThreadPoolWorkQueue when reaching an empty queue state. It's possible for the empty event to be reset after the action count reaches zero and is incremented again.

#rnx

diff --git a/Engine/Source/Programs/DotNETCommon/DotNETUtilities/ThreadPoolWorkQueue.cs b/Engine/Source/Programs/DotNETCommon/DotNETUtilities/ThreadPoolWorkQueue.cs
--- a/Engine/Source/Programs/DotNETCommon/DotNETUtilities/ThreadPoolWorkQueue.cs
+++ b/Engine/Source/Programs/DotNETCommon/DotNETUtilities/ThreadPoolWorkQueue.cs
@@ -1,4 +1,4 @@
-﻿// Copyright 1998-2019 Epic Games, Inc. All Rights Reserved.
+// Copyright 1998-2019 Epic Games, Inc. All Rights Reserved.
 
 using System;
 using System.Collections.Generic;
@@ -14,6 +14,11 @@ namespace Tools.DotNETCommon
 	/// </summary>
 	public class ThreadPoolWorkQueue : IDisposable
 	{
+		/// <summary>
+		/// Object used for controlling access to NumOutstandingJobs and updating EmptyEvent
+		/// </summary>
+		object LockObject = new object();
+
 		/// <summary>
 		/// Number of jobs remaining in the queue. This is updated in an atomic way.
 		/// </summary>
@@ -59,10 +64,15 @@ namespace Tools.DotNETCommon
 		/// <param name="ActionToExecute">The action to add</param>
 		public void Enqueue(Action ActionToExecute)
 		{
-			if(Interlocked.Increment(ref NumOutstandingJobs) == 1)
+			lock(LockObject)
 			{
-				EmptyEvent.Reset();
+				if(NumOutstandingJobs == 0)
+				{
+					EmptyEvent.Reset();
+				}
+				NumOutstandingJobs++;
 			}
+
 			ThreadPool.QueueUserWorkItem(Execute, ActionToExecute);
 		}
 
@@ -73,9 +83,14 @@ namespace Tools.DotNETCommon
 		void Execute(object ActionToExecute)
 		{
 			((Action)ActionToExecute)();
-			if(Interlocked.Decrement(ref NumOutstandingJobs) == 0)
+
+			lock(LockObject)
 			{
-				EmptyEvent.Set();
+				NumOutstandingJobs--;
+				if(NumOutstandingJobs == 0)
+				{
+					EmptyEvent.Set();
+				}
 			}
 		}
 
