Add a C# utility class for computing hashes of files and blocks of data.

#rnx

diff --git a/Engine/Source/Programs/DotNETCommon/DotNETUtilities/ContentHash.cs b/Engine/Source/Programs/DotNETCommon/DotNETUtilities/ContentHash.cs
new file mode 100644
--- /dev/null
+++ b/Engine/Source/Programs/DotNETCommon/DotNETUtilities/ContentHash.cs
@@ -0,0 +1,276 @@
+using System;
+using System.Collections.Generic;
+using System.IO;
+using System.Linq;
+using System.Security.Cryptography;
+using System.Text;
+using System.Threading.Tasks;
+
+namespace Tools.DotNETCommon
+{
+	/// <summary>
+	/// Stores the hash value for a piece of content as a byte array, allowing it to be used as a dictionary key
+	/// </summary>
+	public class ContentHash : IEquatable<ContentHash>
+	{
+		/// <summary>
+		/// The bytes compromising this hash
+		/// </summary>
+		public byte[] Bytes
+		{
+			get;
+			private set;
+		}
+
+		/// <summary>
+		/// Constructor
+		/// </summary>
+		/// <param name="Bytes">The hash data</param>
+		public ContentHash(byte[] Bytes)
+		{
+			this.Bytes = Bytes;
+		}
+
+		/// <summary>
+		/// Compares two content hashes for equality
+		/// </summary>
+		/// <param name="Other">The object to compare against</param>
+		/// <returns>True if the hashes are equal, false otherwise</returns>
+		public override bool Equals(object Other)
+		{
+			return Equals(Other as ContentHash);
+		}
+
+		/// <summary>
+		/// Compares two content hashes for equality
+		/// </summary>
+		/// <param name="Other">The hash to compare against</param>
+		/// <returns>True if the hashes are equal, false otherwise</returns>
+		public bool Equals(ContentHash Other)
+		{
+			if(Other == null)
+			{
+				return false;
+			}
+			if(Bytes.Length != Other.Bytes.Length)
+			{
+				return false;
+			}
+			for(int Idx = 0; Idx < Bytes.Length; Idx++)
+			{
+				if(Bytes[Idx] != Other.Bytes[Idx])
+				{
+					return false;
+				}
+			}
+			return true;
+		}
+
+		/// <summary>
+		/// Compares two content hash objects for equality
+		/// </summary>
+		/// <param name="A">The first hash to compare</param>
+		/// <param name="B">The second has to compare</param>
+		/// <returns>True if the objects are equal, false otherwise</returns>
+		public static bool operator==(ContentHash A, ContentHash B)
+		{
+			if(A == null)
+			{
+				return (B == null);
+			}
+			else
+			{
+				return A.Equals(B);
+			}
+		}
+
+		/// <summary>
+		/// Compares two content hash objects for inequality
+		/// </summary>
+		/// <param name="A">The first hash to compare</param>
+		/// <param name="B">The second has to compare</param>
+		/// <returns>True if the objects are not equal, false otherwise</returns>
+		public static bool operator!=(ContentHash A, ContentHash B)
+		{
+			return !(A == B);
+		}
+
+		/// <summary>
+		/// Creates a content hash for a block of data, using a given algorithm.
+		/// </summary>
+		/// <param name="Data">Data to compute the hash for</param>
+		/// <param name="Algorithm">Algorithm to use to create the hash</param>
+		/// <returns>New content hash instance containing the hash of the data</returns>
+		public static ContentHash Compute(byte[] Data, HashAlgorithm Algorithm)
+		{
+			return new ContentHash(Algorithm.ComputeHash(Data));
+		}
+
+		/// <summary>
+		/// Creates a content hash for a string, using a given algorithm.
+		/// </summary>
+		/// <param name="Text">Text to compute a hash for</param>
+		/// <param name="Algorithm">Algorithm to use to create the hash</param>
+		/// <returns>New content hash instance containing the hash of the text</returns>
+		public static ContentHash Compute(string Text, HashAlgorithm Algorithm)
+		{
+			return new ContentHash(Algorithm.ComputeHash(Encoding.Unicode.GetBytes(Text)));
+		}
+
+		/// <summary>
+		/// Creates a content hash for a file, using a given algorithm.
+		/// </summary>
+		/// <param name="Location">File to compute a hash for</param>
+		/// <param name="Algorithm">Algorithm to use to create the hash</param>
+		/// <returns>New content hash instance containing the hash of the file</returns>
+		public static ContentHash Compute(FileReference Location, HashAlgorithm Algorithm)
+		{
+			using(FileStream Stream = FileReference.Open(Location, FileMode.Open, FileAccess.Read, FileShare.Read))
+			{
+				return new ContentHash(Algorithm.ComputeHash(Stream));
+			}
+		}
+
+		/// <summary>
+		/// Creates a content hash for a block of data using MD5
+		/// </summary>
+		/// <param name="Data">Data to compute the hash for</param>
+		/// <returns>New content hash instance containing the hash of the data</returns>
+		public static ContentHash MD5(byte[] Data)
+		{
+			using(MD5 Algorithm = System.Security.Cryptography.MD5.Create())
+			{
+				return Compute(Data, Algorithm);
+			}
+		}
+
+		/// <summary>
+		/// Creates a content hash for a string using MD5.
+		/// </summary>
+		/// <param name="Text">Text to compute a hash for</param>
+		/// <returns>New content hash instance containing the hash of the text</returns>
+		public static ContentHash MD5(string Text)
+		{
+			using(MD5 Algorithm = System.Security.Cryptography.MD5.Create())
+			{
+				return Compute(Text, Algorithm);
+			}
+		}
+
+		/// <summary>
+		/// Creates a content hash for a file, using a given algorithm.
+		/// </summary>
+		/// <param name="Location">File to compute a hash for</param>
+		/// <returns>New content hash instance containing the hash of the file</returns>
+		public static ContentHash MD5(FileReference Location)
+		{
+			using(MD5 Algorithm = System.Security.Cryptography.MD5.Create())
+			{
+				return Compute(Location, Algorithm);
+			}
+		}
+
+		/// <summary>
+		/// Creates a content hash for a block of data using SHA1
+		/// </summary>
+		/// <param name="Data">Data to compute the hash for</param>
+		/// <returns>New content hash instance containing the hash of the data</returns>
+		public static ContentHash SHA1(byte[] Data)
+		{
+			using(SHA1 Algorithm = System.Security.Cryptography.SHA1.Create())
+			{
+				return Compute(Data, Algorithm);
+			}
+		}
+
+		/// <summary>
+		/// Creates a content hash for a string using SHA1.
+		/// </summary>
+		/// <param name="Text">Text to compute a hash for</param>
+		/// <returns>New content hash instance containing the hash of the text</returns>
+		public static ContentHash SHA1(string Text)
+		{
+			using(SHA1 Algorithm = System.Security.Cryptography.SHA1.Create())
+			{
+				return Compute(Text, Algorithm);
+			}
+		}
+
+		/// <summary>
+		/// Creates a content hash for a file using SHA1.
+		/// </summary>
+		/// <param name="Location">File to compute a hash for</param>
+		/// <returns>New content hash instance containing the hash of the file</returns>
+		public static ContentHash SHA1(FileReference Location)
+		{
+			using(SHA1 Algorithm = System.Security.Cryptography.SHA1.Create())
+			{
+				return Compute(Location, Algorithm);
+			}
+		}
+
+		/// <summary>
+		/// Computes a hash code for this digest
+		/// </summary>
+		/// <returns>Integer value to use as a hash code</returns>
+		public override int GetHashCode()
+		{
+			int HashCode = Bytes[0];
+			for(int Idx = 1; Idx < Bytes.Length; Idx++)
+			{
+				HashCode = (HashCode * 31) + Bytes[Idx];
+			}
+			return HashCode;
+		}
+
+		/// <summary>
+		/// Formats this hash as a string
+		/// </summary>
+		/// <returns>The hashed value</returns>
+		public override string ToString()
+		{
+			return StringUtils.FormatHexString(Bytes);
+		}
+	}
+
+	/// <summary>
+	/// Utility methods for serializing ContentHash objects
+	/// </summary>
+	public static class ContentHashExtensionMethods
+	{
+		/// <summary>
+		/// Writes a ContentHash to a binary archive
+		/// </summary>
+		/// <param name="Writer">The writer to output data to</param>
+		/// <param name="Hash">The hash to write</param>
+		public static void WriteContentHash(this BinaryArchiveWriter Writer, ContentHash Hash)
+		{
+			if(Hash == null)
+			{
+				Writer.WriteByteArray(null);
+			}
+			else
+			{
+				Writer.WriteByteArray(Hash.Bytes);
+			}
+		}
+
+		/// <summary>
+		/// Reads a ContentHash from a binary archive
+		/// </summary>
+		/// <param name="Reader">Reader to serialize data from</param>
+		/// <returns>New hash instance</returns>
+		public static ContentHash ReadContentHash(this BinaryArchiveReader Reader)
+		{
+			byte[] Data = Reader.ReadByteArray();
+			if(Data == null)
+			{
+				return null;
+			}
+			else
+			{
+				return new ContentHash(Data);
+			}
+		}
+	}
+}
diff --git a/Engine/Source/Programs/DotNETCommon/DotNETUtilities/DotNETUtilities.csproj b/Engine/Source/Programs/DotNETCommon/DotNETUtilities/DotNETUtilities.csproj
--- a/Engine/Source/Programs/DotNETCommon/DotNETUtilities/DotNETUtilities.csproj
+++ b/Engine/Source/Programs/DotNETCommon/DotNETUtilities/DotNETUtilities.csproj
@@ -63,6 +63,7 @@
     <Compile Include="CommandLineArguments.cs" />
     <Compile Include="CommandLineAttribute.cs" />
     <Compile Include="CsProjectInfo.cs" />
+    <Compile Include="ContentHash.cs" />
     <Compile Include="DirectoryReference.cs" />
     <Compile Include="ExceptionUtils.cs" />
     <Compile Include="FastJSON\Formatter.cs" />
