Missing file.

#rnx

diff --git a/Engine/Source/Programs/DotNETCommon/DotNETUtilities/StringUtils.cs b/Engine/Source/Programs/DotNETCommon/DotNETUtilities/StringUtils.cs
--- a/Engine/Source/Programs/DotNETCommon/DotNETUtilities/StringUtils.cs
+++ b/Engine/Source/Programs/DotNETCommon/DotNETUtilities/StringUtils.cs
@@ -73,5 +73,23 @@ namespace Tools.DotNETCommon
 		{
 			return FormatList(Arguments.ToArray());
 		}
+
+		/// <summary>
+		/// Formats an array of bytes as a hexadecimal string
+		/// </summary>
+		/// <param name="Bytes">An array of bytes</param>
+		/// <returns>String representation of the array</returns>
+		public static string FormatHexString(byte[] Bytes)
+		{
+			const string HexDigits = "0123456789abcdef";
+
+			char[] Characters = new char[Bytes.Length * 2];
+			for(int Idx = 0; Idx < Bytes.Length; Idx++)
+			{
+				Characters[Idx * 2 + 0] = HexDigits[Bytes[Idx] >> 4];
+				Characters[Idx * 2 + 1] = HexDigits[Bytes[Idx] & 15];
+			}
+			return new string(Characters);
+		}
 	}
 }
