Add missing copyright notice.

#rnx

diff --git a/Engine/Source/Programs/DotNETCommon/DotNETUtilities/ContentHash.cs b/Engine/Source/Programs/DotNETCommon/DotNETUtilities/ContentHash.cs
--- a/Engine/Source/Programs/DotNETCommon/DotNETUtilities/ContentHash.cs
+++ b/Engine/Source/Programs/DotNETCommon/DotNETUtilities/ContentHash.cs
@@ -1,3 +1,5 @@
+// Copyright 1998-2019 Epic Games, Inc. All Rights Reserved.
+
 using System;
 using System.Collections.Generic;
 using System.IO;
