Fix stack overflow when comparing two ContentHash objects.

#rnx

diff --git a/Engine/Source/Programs/DotNETCommon/DotNETUtilities/ContentHash.cs b/Engine/Source/Programs/DotNETCommon/DotNETUtilities/ContentHash.cs
--- a/Engine/Source/Programs/DotNETCommon/DotNETUtilities/ContentHash.cs
+++ b/Engine/Source/Programs/DotNETCommon/DotNETUtilities/ContentHash.cs
@@ -76,9 +76,9 @@ namespace Tools.DotNETCommon
 		/// <returns>True if the objects are equal, false otherwise</returns>
 		public static bool operator==(ContentHash A, ContentHash B)
 		{
-			if(A == null)
+			if((object)A == null)
 			{
-				return (B == null);
+				return ((object)B == null);
 			}
 			else
 			{
