Add a macro to allow forcing ThreadPoolWorkQueue to execute on a single thread for easier debugging.

#rnx

diff --git a/Engine/Source/Programs/DotNETCommon/DotNETUtilities/ThreadPoolWorkQueue.cs b/Engine/Source/Programs/DotNETCommon/DotNETUtilities/ThreadPoolWorkQueue.cs
--- a/Engine/Source/Programs/DotNETCommon/DotNETUtilities/ThreadPoolWorkQueue.cs
+++ b/Engine/Source/Programs/DotNETCommon/DotNETUtilities/ThreadPoolWorkQueue.cs
@@ -73,7 +73,11 @@ namespace Tools.DotNETCommon
 				NumOutstandingJobs++;
 			}
 
+#if SINGLE_THREAD
+			Execute(ActionToExecute);
+#else
 			ThreadPool.QueueUserWorkItem(Execute, ActionToExecute);
+#endif
 		}
 
 		/// <summary>
