Prevent exceptions thrown by tasks queued to ThreadPoolWorkQueue causing state to become out of sync. Now captures all exceptions and re-throws them when Wait() is called.

diff --git a/Engine/Source/Programs/DotNETCommon/DotNETUtilities/ThreadPoolWorkQueue.cs b/Engine/Source/Programs/DotNETCommon/DotNETUtilities/ThreadPoolWorkQueue.cs
--- a/Engine/Source/Programs/DotNETCommon/DotNETUtilities/ThreadPoolWorkQueue.cs
+++ b/Engine/Source/Programs/DotNETCommon/DotNETUtilities/ThreadPoolWorkQueue.cs
@@ -29,6 +29,11 @@ namespace Tools.DotNETCommon
 		/// </summary>
 		ManualResetEvent EmptyEvent = new ManualResetEvent(true);
 
+		/// <summary>
+		/// Exceptions which occurred while executing tasks
+		/// </summary>
+		List<Exception> Exceptions = new List<Exception>();
+
 		/// <summary>
 		/// Default constructor
 		/// </summary>
@@ -86,14 +91,27 @@ namespace Tools.DotNETCommon
 		/// <param name="ActionToExecute">The action to execute</param>
 		void Execute(object ActionToExecute)
 		{
-			((Action)ActionToExecute)();
-
-			lock(LockObject)
+			try
 			{
-				NumOutstandingJobs--;
-				if(NumOutstandingJobs == 0)
+				((Action)ActionToExecute)();
+			}
+			catch(Exception Ex)
+			{
+				lock(LockObject)
+				{
+					Exceptions.Add(Ex);
+				}
+				throw;
+			}
+			finally
+			{
+				lock(LockObject)
 				{
-					EmptyEvent.Set();
+					NumOutstandingJobs--;
+					if(NumOutstandingJobs == 0)
+					{
+						EmptyEvent.Set();
+					}
 				}
 			}
 		}
@@ -104,6 +122,7 @@ namespace Tools.DotNETCommon
 		public void Wait()
 		{
 			EmptyEvent.WaitOne();
+			RethrowExceptions();
 		}
 
 		/// <summary>
@@ -113,7 +132,26 @@ namespace Tools.DotNETCommon
 		/// <returns>True if the queue completed, false if the timeout elapsed</returns>
 		public bool Wait(int MillisecondsTimeout)
 		{
-			return EmptyEvent.WaitOne(MillisecondsTimeout);
+			bool bResult = EmptyEvent.WaitOne(MillisecondsTimeout);
+			if(bResult)
+			{
+				RethrowExceptions();
+			}
+			return bResult;
+		}
+
+		/// <summary>
+		/// Checks for any exceptions which ocurred in queued tasks, and re-throws them on the current thread
+		/// </summary>
+		public void RethrowExceptions()
+		{
+			lock(LockObject)
+			{
+				if(Exceptions.Count > 0)
+				{
+					throw new AggregateException(Exceptions.ToArray());
+				}
+			}
 		}
 	}
 }
