Minor tweak to ContentHash to improve performance of comparisons against null.

#rnx

diff --git a/Engine/Source/Programs/DotNETCommon/DotNETUtilities/ContentHash.cs b/Engine/Source/Programs/DotNETCommon/DotNETUtilities/ContentHash.cs
--- a/Engine/Source/Programs/DotNETCommon/DotNETUtilities/ContentHash.cs
+++ b/Engine/Source/Programs/DotNETCommon/DotNETUtilities/ContentHash.cs
@@ -50,7 +50,7 @@ namespace Tools.DotNETCommon
 		/// <returns>True if the hashes are equal, false otherwise</returns>
 		public bool Equals(ContentHash Other)
 		{
-			if(Other == null)
+			if((object)Other == null)
 			{
 				return false;
 			}
