Add methods for serializing nullable objects and structures to binary archives.

#rnx

diff --git a/Engine/Source/Programs/DotNETCommon/DotNETUtilities/BinaryArchiveReader.cs b/Engine/Source/Programs/DotNETCommon/DotNETUtilities/BinaryArchiveReader.cs
--- a/Engine/Source/Programs/DotNETCommon/DotNETUtilities/BinaryArchiveReader.cs
+++ b/Engine/Source/Programs/DotNETCommon/DotNETUtilities/BinaryArchiveReader.cs
@@ -390,6 +390,41 @@ namespace Tools.DotNETCommon
 			}
 		}
 
+		/// <summary>
+		/// Reads a nullable object from the archive
+		/// </summary>
+		/// <typeparam name="T">The nullable type</typeparam>
+		/// <param name="ReadValue">Delegate used to read a value</param>
+		public Nullable<T> ReadNullable<T>(Func<T> ReadValue) where T : struct
+		{
+			if(ReadBool())
+			{
+				return new Nullable<T>(ReadValue());
+			}
+			else
+			{
+				return null;
+			}
+		}
+
+		/// <summary>
+		/// Reads an object, which may be null, from the archive. Does not handle de-duplicating object references. 
+		/// </summary>
+		/// <typeparam name="T">Type of the object to read</typeparam>
+		/// <param name="Read">Delegate used to read the object</param>
+		/// <returns>The object instance</returns>
+		public T ReadOptionalObject<T>(Func<T> Read) where T : class
+		{
+			if(ReadBool())
+			{
+				return Read();
+			}
+			else
+			{
+				return null;
+			}
+		}
+
 		/// <summary>
 		/// Reads an object reference from the stream. Each referenced object will only be serialized once using the supplied delegates. Reading an object instance is
 		/// done in two phases; first the object is created and its reference stored in the unique object list, then the object contents are read. This allows the object to 
diff --git a/Engine/Source/Programs/DotNETCommon/DotNETUtilities/BinaryArchiveWriter.cs b/Engine/Source/Programs/DotNETCommon/DotNETUtilities/BinaryArchiveWriter.cs
--- a/Engine/Source/Programs/DotNETCommon/DotNETUtilities/BinaryArchiveWriter.cs
+++ b/Engine/Source/Programs/DotNETCommon/DotNETUtilities/BinaryArchiveWriter.cs
@@ -405,6 +405,44 @@ namespace Tools.DotNETCommon
 			}
 		}
 
+		/// <summary>
+		/// Writes a nullable object to the archive
+		/// </summary>
+		/// <typeparam name="T">The nullable type</typeparam>
+		/// <param name="Item">Item to write</param>
+		/// <param name="WriteValue">Delegate used to write a value</param>
+		public void WriteNullable<T>(Nullable<T> Item, Action<T> WriteValue) where T : struct
+		{
+			if(Item.HasValue)
+			{
+				WriteBool(true);
+				WriteValue(Item.Value);
+			}
+			else
+			{
+				WriteBool(false);
+			}
+		}
+
+		/// <summary>
+		/// Writes an object to the output, checking whether it is null or not. Does not preserve object references; each object written is duplicated.
+		/// </summary>
+		/// <typeparam name="T">Type of the object to serialize</typeparam>
+		/// <param name="Object">Reference to check for null before serializing</param>
+		/// <param name="WriteObject">Delegate used to write the object</param>
+		public void WriteOptionalObject<T>(T Object, Action WriteObject) where T : class
+		{
+			if(Object == null)
+			{
+				WriteBool(false);
+			}
+			else
+			{
+				WriteBool(true);
+				WriteObject();
+			}
+		}
+
 		/// <summary>
 		/// Writes an object to the output. If the specific instance has already been written, preserves the reference to that.
 		/// </summary>
