UBT: Add utility methods for correcting paths to match their case on disk, and fix issues with inconsistent paths when cross-compiling for Linux.

diff --git a/Engine/Source/Programs/DotNETCommon/DotNETUtilities/DirectoryReference.cs b/Engine/Source/Programs/DotNETCommon/DotNETUtilities/DirectoryReference.cs
--- a/Engine/Source/Programs/DotNETCommon/DotNETUtilities/DirectoryReference.cs
+++ b/Engine/Source/Programs/DotNETCommon/DotNETUtilities/DirectoryReference.cs
@@ -258,6 +258,25 @@ namespace Tools.DotNETCommon
 			}
 		}
 
+		/// <summary>
+		/// Finds the correct case to match the location of this file on disk. Uses the given case for parts of the path that do not exist.
+		/// </summary>
+		/// <param name="Location">The path to find the correct case for</param>
+		/// <returns>Location of the file with the correct case</returns>
+		public static DirectoryReference FindCorrectCase(DirectoryReference Location)
+		{
+			return new DirectoryReference(DirectoryUtils.FindCorrectCase(Location.ToDirectoryInfo()));
+		}
+
+		/// <summary>
+		/// Constructs a DirectoryInfo object from this reference
+		/// </summary>
+		/// <returns>New DirectoryInfo object</returns>
+		public DirectoryInfo ToDirectoryInfo()
+		{
+			return new DirectoryInfo(FullName);
+		}
+
 		#region System.IO.Directory Wrapper Methods
 
 		/// <summary>
diff --git a/Engine/Source/Programs/DotNETCommon/DotNETUtilities/DirectoryUtils.cs b/Engine/Source/Programs/DotNETCommon/DotNETUtilities/DirectoryUtils.cs
new file mode 100644
--- /dev/null
+++ b/Engine/Source/Programs/DotNETCommon/DotNETUtilities/DirectoryUtils.cs
@@ -0,0 +1,49 @@
+using System;
+using System.Collections.Generic;
+using System.IO;
+using System.Linq;
+using System.Text;
+using System.Threading.Tasks;
+
+namespace Tools.DotNETCommon
+{
+	/// <summary>
+	/// Utility functions for manipulating directories
+	/// </summary>
+	public class DirectoryUtils
+	{
+		/// <summary>
+		/// Finds the on-disk case of a a directory
+		/// </summary>
+		/// <param name="Info">DirectoryInfo instance describing the directory</param>
+		/// <returns>New DirectoryInfo instance that represents the directory with the correct case</returns>
+		public static DirectoryInfo FindCorrectCase(DirectoryInfo Info)
+		{
+			DirectoryInfo ParentInfo = Info.Parent;
+			if (ParentInfo == null)
+			{
+				string FullName = Info.FullName;
+				if (FullName.Length >= 2 && (FullName[0] >= 'a' && FullName[0] <= 'z') && FullName[1] == ':')
+				{
+					return new DirectoryInfo(Char.ToUpper(FullName[0]) + FullName.Substring(1));
+				}
+				else
+				{
+					return Info;
+				}
+			}
+			else
+			{
+				ParentInfo = FindCorrectCase(ParentInfo);
+				foreach (DirectoryInfo ChildInfo in ParentInfo.EnumerateDirectories())
+				{
+					if (String.Equals(ChildInfo.Name, Info.Name, DirectoryReference.Comparison))
+					{
+						return ChildInfo;
+					}
+				}
+				return new DirectoryInfo(Path.Combine(ParentInfo.FullName, Info.Name));
+			}
+		}
+	}
+}
diff --git a/Engine/Source/Programs/DotNETCommon/DotNETUtilities/DotNETUtilities.csproj b/Engine/Source/Programs/DotNETCommon/DotNETUtilities/DotNETUtilities.csproj
--- a/Engine/Source/Programs/DotNETCommon/DotNETUtilities/DotNETUtilities.csproj
+++ b/Engine/Source/Programs/DotNETCommon/DotNETUtilities/DotNETUtilities.csproj
@@ -67,6 +67,7 @@
     <Compile Include="CsProjectInfo.cs" />
     <Compile Include="ContentHash.cs" />
     <Compile Include="DirectoryReference.cs" />
+    <Compile Include="DirectoryUtils.cs" />
     <Compile Include="ExceptionUtils.cs" />
     <Compile Include="FastJSON\Formatter.cs" />
     <Compile Include="FastJSON\Getters.cs" />
@@ -80,6 +81,7 @@
     <Compile Include="FilePattern.cs" />
     <Compile Include="FileReference.cs" />
     <Compile Include="FileSystemReference.cs" />
+    <Compile Include="FileUtils.cs" />
     <Compile Include="Json.cs" />
     <Compile Include="JsonObject.cs" />
     <Compile Include="JsonWriter.cs" />
diff --git a/Engine/Source/Programs/DotNETCommon/DotNETUtilities/FileReference.cs b/Engine/Source/Programs/DotNETCommon/DotNETUtilities/FileReference.cs
--- a/Engine/Source/Programs/DotNETCommon/DotNETUtilities/FileReference.cs
+++ b/Engine/Source/Programs/DotNETCommon/DotNETUtilities/FileReference.cs
@@ -246,6 +246,16 @@ namespace Tools.DotNETCommon
 			}
 		}
 
+		/// <summary>
+		/// Finds the correct case to match the location of this file on disk. Uses the given case for parts of the path that do not exist.
+		/// </summary>
+		/// <param name="Location">The path to find the correct case for</param>
+		/// <returns>Location of the file with the correct case</returns>
+		public static FileReference FindCorrectCase(FileReference Location)
+		{
+			return new FileReference(FileUtils.FindCorrectCase(Location.ToFileInfo()));
+		}
+
 		/// <summary>
 		/// Constructs a FileInfo object from this reference
 		/// </summary>
diff --git a/Engine/Source/Programs/DotNETCommon/DotNETUtilities/FileUtils.cs b/Engine/Source/Programs/DotNETCommon/DotNETUtilities/FileUtils.cs
new file mode 100644
--- /dev/null
+++ b/Engine/Source/Programs/DotNETCommon/DotNETUtilities/FileUtils.cs
@@ -0,0 +1,33 @@
+using System;
+using System.Collections.Generic;
+using System.IO;
+using System.Linq;
+using System.Text;
+using System.Threading.Tasks;
+
+namespace Tools.DotNETCommon
+{
+	/// <summary>
+	/// Utility functions for manipulating files
+	/// </summary>
+	public static class FileUtils
+	{
+		/// <summary>
+		/// Finds the on-disk case of a a file
+		/// </summary>
+		/// <param name="Info">FileInfo instance describing the file</param>
+		/// <returns>New FileInfo instance that represents the file with the correct case</returns>
+		public static FileInfo FindCorrectCase(FileInfo Info)
+		{
+			DirectoryInfo ParentInfo = DirectoryUtils.FindCorrectCase(Info.Directory);
+			foreach (FileInfo ChildInfo in ParentInfo.EnumerateFiles())
+			{
+				if (String.Equals(ChildInfo.Name, Info.Name, FileReference.Comparison))
+				{
+					return ChildInfo;
+				}
+			}
+			return new FileInfo(Path.Combine(ParentInfo.FullName, Info.Name));
+		}
+	}
+}
diff --git a/Engine/Source/Programs/UnrealBuildTool/Configuration/TargetDescriptor.cs b/Engine/Source/Programs/UnrealBuildTool/Configuration/TargetDescriptor.cs
--- a/Engine/Source/Programs/UnrealBuildTool/Configuration/TargetDescriptor.cs
+++ b/Engine/Source/Programs/UnrealBuildTool/Configuration/TargetDescriptor.cs
@@ -259,10 +259,15 @@ namespace UnrealBuildTool
 				throw new BuildException("No configurations specified for target");
 			}
 
-			// Make sure the project file exists
-			if(ProjectFile != null && !FileReference.Exists(ProjectFile))
+			// Make sure the project file exists, and make sure we're using the correct case.
+			if(ProjectFile != null)
 			{
-				throw new BuildException("Unable to find project '{0}'.", ProjectFile);
+				FileInfo ProjectFileInfo = FileUtils.FindCorrectCase(ProjectFile.ToFileInfo());
+				if(!ProjectFileInfo.Exists)
+				{
+					throw new BuildException("Unable to find project '{0}'.", ProjectFile);
+				}
+				ProjectFile = new FileReference(ProjectFileInfo);
 			}
 
 			// Expand all the platforms, architectures and configurations
