Add missing copyright notices.

#rnx

diff --git a/Engine/Source/Programs/DotNETCommon/DotNETUtilities/DirectoryUtils.cs b/Engine/Source/Programs/DotNETCommon/DotNETUtilities/DirectoryUtils.cs
--- a/Engine/Source/Programs/DotNETCommon/DotNETUtilities/DirectoryUtils.cs
+++ b/Engine/Source/Programs/DotNETCommon/DotNETUtilities/DirectoryUtils.cs
@@ -1,3 +1,5 @@
+// Copyright 1998-2019 Epic Games, Inc. All Rights Reserved.
+
 using System;
 using System.Collections.Generic;
 using System.IO;
diff --git a/Engine/Source/Programs/DotNETCommon/DotNETUtilities/FileUtils.cs b/Engine/Source/Programs/DotNETCommon/DotNETUtilities/FileUtils.cs
--- a/Engine/Source/Programs/DotNETCommon/DotNETUtilities/FileUtils.cs
+++ b/Engine/Source/Programs/DotNETCommon/DotNETUtilities/FileUtils.cs
@@ -1,3 +1,5 @@
+// Copyright 1998-2019 Epic Games, Inc. All Rights Reserved.
+
 using System;
 using System.Collections.Generic;
 using System.IO;
