Add utility methods for parsing ContentHash values.

#rnx

diff --git a/Engine/Source/Programs/DotNETCommon/DotNETUtilities/ContentHash.cs b/Engine/Source/Programs/DotNETCommon/DotNETUtilities/ContentHash.cs
--- a/Engine/Source/Programs/DotNETCommon/DotNETUtilities/ContentHash.cs
+++ b/Engine/Source/Programs/DotNETCommon/DotNETUtilities/ContentHash.cs
@@ -15,6 +15,9 @@ namespace Tools.DotNETCommon
 	/// </summary>
 	public class ContentHash : IEquatable<ContentHash>
 	{
+		public const int LengthMD5 = 16;
+		public const int LengthSHA1 = 20;
+
 		/// <summary>
 		/// The bytes compromising this hash
 		/// </summary>
@@ -211,6 +214,41 @@ namespace Tools.DotNETCommon
 			}
 		}
 
+		/// <summary>
+		/// Parse a hash from a string
+		/// </summary>
+		/// <param name="Text">Text to parse</param>
+		/// <returns>Value of the hash</returns>
+		public static ContentHash Parse(string Text)
+		{
+			ContentHash Hash;
+			if(!TryParse(Text, out Hash))
+			{
+				throw new ArgumentException(String.Format("'{0}' is not a valid content hash", Text));
+			}
+			return Hash;
+		}
+
+		/// <summary>
+		/// Parse a hash from a string
+		/// </summary>
+		/// <param name="Text">Text to parse</param>
+		/// <returns>Value of the hash</returns>
+		public static bool TryParse(string Text, out ContentHash Hash)
+		{
+			byte[] Bytes;
+			if(StringUtils.TryParseHexString(Text, out Bytes))
+			{
+				Hash = new ContentHash(Bytes);
+				return true;
+			}
+			else
+			{
+				Hash = null;
+				return false;
+			}
+		}
+
 		/// <summary>
 		/// Computes a hash code for this digest
 		/// </summary>
diff --git a/Engine/Source/Programs/DotNETCommon/DotNETUtilities/StringUtils.cs b/Engine/Source/Programs/DotNETCommon/DotNETUtilities/StringUtils.cs
--- a/Engine/Source/Programs/DotNETCommon/DotNETUtilities/StringUtils.cs
+++ b/Engine/Source/Programs/DotNETCommon/DotNETUtilities/StringUtils.cs
@@ -74,6 +74,57 @@ namespace Tools.DotNETCommon
 			return FormatList(Arguments.ToArray());
 		}
 
+		/// <summary>
+		/// Parses a hexadecimal digit
+		/// </summary>
+		/// <param name="Character">Character to parse</param>
+		/// <returns>Value of this digit</returns>
+		public static bool TryParseHexDigit(char Character, out int Value)
+		{
+			if(Character >= '0' && Character <= '9')
+			{
+				Value = Character - '0';
+				return true;
+			}
+			else if(Character >= 'a' && Character <= 'f')
+			{
+				Value = 10 + (Character - 'a');
+				return true;
+			}
+			else if(Character >= 'A' && Character <= 'F')
+			{
+				Value = 10 + (Character - 'A');
+				return true;
+			}
+			else
+			{
+				Value = 0;
+				return false;
+			}
+		}
+
+		/// <summary>
+		/// Parses a hexadecimal string into an array of bytes
+		/// </summary>
+		/// <param name="Text">Text to parse</param>
+		/// <returns></returns>
+		public static bool TryParseHexString(string Text, out byte[] OutBytes)
+		{
+			byte[] Bytes = new byte[Text.Length / 2];
+			for(int Idx = 0; Idx < Text.Length; Idx++)
+			{
+				int A, B;
+				if(!TryParseHexDigit(Text[Idx], out A) || !TryParseHexDigit(Text[Idx + 1], out B))
+				{
+					OutBytes = null;
+					return false;
+				}
+				Bytes[Idx / 2] = (byte)((A << 4) | B);
+			}
+			OutBytes = Bytes;
+			return true;
+		}
+
 		/// <summary>
 		/// Formats an array of bytes as a hexadecimal string
 		/// </summary>
