Add a helper method to write a null object reference to a binary archive.

#rnx

diff --git a/Engine/Source/Programs/DotNETCommon/DotNETUtilities/BinaryArchiveWriter.cs b/Engine/Source/Programs/DotNETCommon/DotNETUtilities/BinaryArchiveWriter.cs
--- a/Engine/Source/Programs/DotNETCommon/DotNETUtilities/BinaryArchiveWriter.cs
+++ b/Engine/Source/Programs/DotNETCommon/DotNETUtilities/BinaryArchiveWriter.cs
@@ -479,6 +479,14 @@ namespace Tools.DotNETCommon
 			}
 		}
 
+		/// <summary>
+		/// Writes a null object reference to the output.
+		/// </summary>
+		public void WriteNullObjectReference()
+		{
+			WriteInt(-1);
+		}
+
 		/// <summary>
 		/// Writes an object to the output. If the specific instance has already been written, preserves the reference to that.
 		/// </summary>
