UGS: Disable support for parallel syncing; seems to be causing errors now that we have support for wildcard syncs.

#rnx

diff --git a/Engine/Source/Programs/DotNETCommon/DotNETUtilities/Perforce/PerforceConnection.cs b/Engine/Source/Programs/DotNETCommon/DotNETUtilities/Perforce/PerforceConnection.cs
--- a/Engine/Source/Programs/DotNETCommon/DotNETUtilities/Perforce/PerforceConnection.cs
+++ b/Engine/Source/Programs/DotNETCommon/DotNETUtilities/Perforce/PerforceConnection.cs
@@ -2062,7 +2062,7 @@ namespace Tools.DotNETCommon.Perforce
 			{
 				Arguments.AppendFormat(" -m {0}", MaxFiles);
 			}
-			if(NumThreads != -1)
+/*			if(NumThreads != -1)
 			{
 				Arguments.AppendFormat(" --parallel-threads={0}", NumThreads);
 				if(Batch != -1)
@@ -2081,7 +2081,7 @@ namespace Tools.DotNETCommon.Perforce
 				{
 					Arguments.AppendFormat(",minsize={0}", MinSize);
 				}
-			}
+			}*/
 			foreach (string FileSpec in FileSpecs)
 			{
 				Arguments.AppendFormat(" \"{0}\"", FileSpec);
diff --git a/Engine/Source/Programs/UnrealGameSync/UnrealGameSync/Properties/AssemblyInfo.cs b/Engine/Source/Programs/UnrealGameSync/UnrealGameSync/Properties/AssemblyInfo.cs
--- a/Engine/Source/Programs/UnrealGameSync/UnrealGameSync/Properties/AssemblyInfo.cs
+++ b/Engine/Source/Programs/UnrealGameSync/UnrealGameSync/Properties/AssemblyInfo.cs
@@ -34,5 +34,5 @@ using System.Runtime.InteropServices;
 // You can specify all the values or you can default the Build and Revision Numbers 
 // by using the '*' as shown below:
 // [assembly: AssemblyVersion("1.0.*")]
-[assembly: AssemblyVersion("1.177.*")]
-[assembly: AssemblyFileVersion("1.177.0.0")]
+[assembly: AssemblyVersion("1.178.*")]
+[assembly: AssemblyFileVersion("1.178.0.0")]
