Fix extra blank line at the end of ManagedProcess output.

diff --git a/Engine/Source/Programs/DotNETCommon/DotNETUtilities/ManagedProcess.cs b/Engine/Source/Programs/DotNETCommon/DotNETUtilities/ManagedProcess.cs
--- a/Engine/Source/Programs/DotNETCommon/DotNETUtilities/ManagedProcess.cs
+++ b/Engine/Source/Programs/DotNETCommon/DotNETUtilities/ManagedProcess.cs
@@ -585,7 +585,7 @@ namespace Tools.DotNETCommon
 				int NumBytesRead = Read(Buffer, NumBytesInBuffer, Buffer.Length - NumBytesInBuffer);
 				if(NumBytesRead == 0)
 				{
-					if(NumBytesInBuffer < Buffer.Length)
+					if(NumBytesInBuffer > 0)
 					{
 						OutputLines.Add(Console.OutputEncoding.GetString(Buffer, 0, NumBytesInBuffer));
 					}
