Fix command line arguments not being split by delimiter characters, and revert CL 4737886.

diff --git a/Engine/Source/Programs/DotNETCommon/DotNETUtilities/CommandLineArguments.cs b/Engine/Source/Programs/DotNETCommon/DotNETUtilities/CommandLineArguments.cs
--- a/Engine/Source/Programs/DotNETCommon/DotNETUtilities/CommandLineArguments.cs
+++ b/Engine/Source/Programs/DotNETCommon/DotNETUtilities/CommandLineArguments.cs
@@ -639,9 +639,22 @@ namespace Tools.DotNETCommon
 								}
 
 								// Apply the value to the field
-								if(ApplyArgument(TargetObject, FieldInfo, Argument, ValueText, AssignedArgument))
+								if(Attribute.ListSeparator == 0)
 								{
-									AssignedArgument = Argument;
+									if(ApplyArgument(TargetObject, FieldInfo, Argument, ValueText, AssignedArgument))
+									{
+										AssignedArgument = Argument;
+									}
+								}
+								else
+								{
+									foreach(string ItemValueText in ValueText.Split(Attribute.ListSeparator))
+									{
+										if(ApplyArgument(TargetObject, FieldInfo, Argument, ItemValueText, AssignedArgument))
+										{
+											AssignedArgument = Argument;
+										}
+									}
 								}
 
 								// Mark this argument as used
diff --git a/Engine/Source/Programs/UnrealBuildTool/Platform/Android/AndroidToolChain.cs b/Engine/Source/Programs/UnrealBuildTool/Platform/Android/AndroidToolChain.cs
--- a/Engine/Source/Programs/UnrealBuildTool/Platform/Android/AndroidToolChain.cs
+++ b/Engine/Source/Programs/UnrealBuildTool/Platform/Android/AndroidToolChain.cs
@@ -143,18 +143,12 @@ namespace UnrealBuildTool
 
 			if (InAdditionalArches != null)
 			{
-				foreach (string Arch in InAdditionalArches)
-				{
-					AdditionalArches.AddRange(Arch.Split('+'));
-				}
+				AdditionalArches.AddRange(InAdditionalArches);
 			}
 
 			if (InAdditionalGPUArches != null)
 			{
-				foreach (string GPU in InAdditionalGPUArches)
-				{
-					AdditionalGPUArches.AddRange(GPU.Split('+'));
-				}
+				AdditionalGPUArches.AddRange(InAdditionalGPUArches);
 			}
 
 			// by default tools chains don't parse arguments, but we want to be able to check the -architectures flag defined above. This is
