Move the ThreadPoolWorkQueue class into DotNETUtilities so it can be used by other projects.

diff --git a/Engine/Source/Programs/AutomationTool/AutomationUtils/AutomationUtils.Automation.csproj b/Engine/Source/Programs/AutomationTool/AutomationUtils/AutomationUtils.Automation.csproj
--- a/Engine/Source/Programs/AutomationTool/AutomationUtils/AutomationUtils.Automation.csproj
+++ b/Engine/Source/Programs/AutomationTool/AutomationUtils/AutomationUtils.Automation.csproj
@@ -82,7 +82,6 @@
     <Compile Include="StagedDirectoryReference.cs" />
     <Compile Include="StagedFileReference.cs" />
     <Compile Include="StagedFileSystemReference.cs" />
-    <Compile Include="ThreadPoolWorkQueue.cs" />
     <Compile Include="UBTUtils.cs" />
     <Compile Include="UE4Build.cs" />
     <Compile Include="Utils.cs" />
diff --git a/Engine/Source/Programs/AutomationTool/AutomationUtils/ThreadPoolWorkQueue.cs b/Engine/Source/Programs/AutomationTool/AutomationUtils/ThreadPoolWorkQueue.cs
deleted file mode 100644
diff --git a/Engine/Source/Programs/AutomationTool/Scripts/CheckForHacks.cs b/Engine/Source/Programs/AutomationTool/Scripts/CheckForHacks.cs
--- a/Engine/Source/Programs/AutomationTool/Scripts/CheckForHacks.cs
+++ b/Engine/Source/Programs/AutomationTool/Scripts/CheckForHacks.cs
@@ -5,6 +5,7 @@ using System.Linq;
 using System.Text;
 using System.Text.RegularExpressions;
 using System.Threading.Tasks;
+using Tools.DotNETCommon;
 
 namespace AutomationTool
 {
diff --git a/Engine/Source/Programs/AutomationTool/Scripts/SyncDDC.cs b/Engine/Source/Programs/AutomationTool/Scripts/SyncDDC.cs
--- a/Engine/Source/Programs/AutomationTool/Scripts/SyncDDC.cs
+++ b/Engine/Source/Programs/AutomationTool/Scripts/SyncDDC.cs
@@ -9,6 +9,7 @@ using System.Text;
 using System.Text.RegularExpressions;
 using System.Threading;
 using System.Threading.Tasks;
+using Tools.DotNETCommon;
 
 namespace AutomationTool
 {
diff --git a/Engine/Source/Programs/DotNETCommon/DotNETUtilities/DotNETUtilities.csproj b/Engine/Source/Programs/DotNETCommon/DotNETUtilities/DotNETUtilities.csproj
--- a/Engine/Source/Programs/DotNETCommon/DotNETUtilities/DotNETUtilities.csproj
+++ b/Engine/Source/Programs/DotNETCommon/DotNETUtilities/DotNETUtilities.csproj
@@ -78,6 +78,7 @@
     <Compile Include="ResXResource\UEResXWriter.cs" />
     <Compile Include="RPCCommandHelper.cs" />
     <Compile Include="StringUtils.cs" />
+    <Compile Include="ThreadPoolWorkQueue.cs" />
   </ItemGroup>
   <Import Project="$(MSBuildToolsPath)\Microsoft.CSharp.targets" />
   <!-- To modify your build process, add your task inside one of the targets below and uncomment it. 
diff --git a/Engine/Source/Programs/DotNETCommon/DotNETUtilities/ThreadPoolWorkQueue.cs b/Engine/Source/Programs/DotNETCommon/DotNETUtilities/ThreadPoolWorkQueue.cs
new file mode 100644
--- /dev/null
+++ b/Engine/Source/Programs/DotNETCommon/DotNETUtilities/ThreadPoolWorkQueue.cs
@@ -0,0 +1,100 @@
+﻿// Copyright 1998-2018 Epic Games, Inc. All Rights Reserved.
+
+using System;
+using System.Collections.Generic;
+using System.Linq;
+using System.Text;
+using System.Threading;
+using System.Threading.Tasks;
+
+namespace Tools.DotNETCommon
+{
+	/// <summary>
+	/// Allows queuing a large number of tasks to a thread pool and waiting for them all to complete.
+	/// </summary>
+	public class ThreadPoolWorkQueue : IDisposable
+	{
+		/// <summary>
+		/// Number of jobs remaining in the queue. This is updated in an atomic way.
+		/// </summary>
+		int NumOutstandingJobs;
+
+		/// <summary>
+		/// Event which indicates whether the queue is empty.
+		/// </summary>
+		ManualResetEvent EmptyEvent = new ManualResetEvent(true);
+
+		/// <summary>
+		/// Default constructor
+		/// </summary>
+		public ThreadPoolWorkQueue()
+		{
+		}
+
+		/// <summary>
+		/// Waits for the queue to drain, and disposes of it
+		/// </summary>
+		public void Dispose()
+		{
+			if(EmptyEvent != null)
+			{
+				Wait();
+
+				EmptyEvent.Dispose();
+				EmptyEvent = null;
+			}
+		}
+
+		/// <summary>
+		/// Returns the number of items remaining in the queue
+		/// </summary>
+		public int NumRemaining
+		{
+			get { return NumOutstandingJobs; }
+		}
+
+		/// <summary>
+		/// Adds an item to the queue
+		/// </summary>
+		/// <param name="ActionToExecute">The action to add</param>
+		public void Enqueue(Action ActionToExecute)
+		{
+			if(Interlocked.Increment(ref NumOutstandingJobs) == 1)
+			{
+				EmptyEvent.Reset();
+			}
+			ThreadPool.QueueUserWorkItem(Execute, ActionToExecute);
+		}
+
+		/// <summary>
+		/// Internal method to execute an action
+		/// </summary>
+		/// <param name="ActionToExecute">The action to execute</param>
+		void Execute(object ActionToExecute)
+		{
+			((Action)ActionToExecute)();
+			if(Interlocked.Decrement(ref NumOutstandingJobs) == 0)
+			{
+				EmptyEvent.Set();
+			}
+		}
+
+		/// <summary>
+		/// Waits for all queued tasks to finish
+		/// </summary>
+		public void Wait()
+		{
+			EmptyEvent.WaitOne();
+		}
+
+		/// <summary>
+		/// Waits for all queued tasks to finish, or the timeout to elapse
+		/// </summary>
+		/// <param name="MillisecondsTimeout">Maximum time to wait</param>
+		/// <returns>True if the queue completed, false if the timeout elapsed</returns>
+		public bool Wait(int MillisecondsTimeout)
+		{
+			return EmptyEvent.WaitOne(MillisecondsTimeout);
+		}
+	}
+}
