Add a IBinarySerializable interface for types that support reading from a BinaryReader and writing to a BinaryWriter. Implementing IBinarySerializable implies a constructor taking a BinaryReader argument is available for deserializing.

diff --git a/Engine/Source/Programs/DotNETCommon/DotNETUtilities/BinaryReaderExtensions.cs b/Engine/Source/Programs/DotNETCommon/DotNETUtilities/BinaryReaderExtensions.cs
--- a/Engine/Source/Programs/DotNETCommon/DotNETUtilities/BinaryReaderExtensions.cs
+++ b/Engine/Source/Programs/DotNETCommon/DotNETUtilities/BinaryReaderExtensions.cs
@@ -11,6 +11,38 @@ namespace Tools.DotNETCommon
 {
 	public static class BinaryReaderExtensions
 	{
+		/// <summary>
+		/// Read an objects from a binary reader
+		/// </summary>
+		/// <typeparam name="T">The element type for the array</typeparam>
+		/// <param name="Reader">Reader to read data from</param>
+		/// <returns>Object instance.</returns>
+		public static T ReadObject<T>(this BinaryReader Reader) where T : IBinarySerializable
+		{
+			return (T)Activator.CreateInstance(typeof(T), Reader);
+		}
+
+		/// <summary>
+		/// Read an array of strings from a binary reader
+		/// </summary>
+		/// <param name="Reader">Reader to read data from</param>
+		/// <returns>Array of strings, as serialized. May be null.</returns>
+		public static string[] ReadStringArray(this BinaryReader Reader)
+		{
+			return Reader.ReadArray(() => Reader.ReadString());
+		}
+
+		/// <summary>
+		/// Read an array of objects from a binary reader
+		/// </summary>
+		/// <typeparam name="T">The element type for the array</typeparam>
+		/// <param name="Reader">Reader to read data from</param>
+		/// <returns>Array of objects, as serialized. May be null.</returns>
+		public static T[] ReadArray<T>(this BinaryReader Reader) where T : IBinarySerializable
+		{
+			return ReadArray<T>(Reader, () => Reader.ReadObject<T>());
+		}
+
 		/// <summary>
 		/// Read an array of objects from a binary reader
 		/// </summary>
@@ -18,7 +50,7 @@ namespace Tools.DotNETCommon
 		/// <param name="Reader">Reader to read data from</param>
 		/// <param name="ReadElement">Delegate to call to serialize each element</param>
 		/// <returns>Array of objects, as serialized. May be null.</returns>
-		public static T[] ReadArray<T>(this BinaryReader Reader, Func<BinaryReader, T> ReadElement)
+		public static T[] ReadArray<T>(this BinaryReader Reader, Func<T> ReadElement)
 		{
 			int NumItems = Reader.ReadInt32();
 			if(NumItems < 0)
@@ -29,11 +61,21 @@ namespace Tools.DotNETCommon
 			T[] Items = new T[NumItems];
 			for(int Idx = 0; Idx < NumItems; Idx++)
 			{
-				Items[Idx] = ReadElement(Reader);
+				Items[Idx] = ReadElement();
 			}
 			return Items;
 		}
 
+		/// <summary>
+		/// Read a list of strings from a binary reader
+		/// </summary>
+		/// <param name="Reader">Reader to read data from</param>
+		/// <returns>Array of strings, as serialized. May be null.</returns>
+		public static List<string> ReadStringList(this BinaryReader Reader)
+		{
+			return Reader.ReadList(() => Reader.ReadString());
+		}
+
 		/// <summary>
 		/// Read a list of objects from a binary reader
 		/// </summary>
@@ -41,10 +83,58 @@ namespace Tools.DotNETCommon
 		/// <param name="Reader">Reader to read data from</param>
 		/// <param name="ReadElement">Delegate to call to serialize each element</param>
 		/// <returns>List of objects, as serialized. May be null.</returns>
-		public static List<T> ReadList<T>(this BinaryReader Reader, Func<BinaryReader, T> ReadElement)
+		public static List<T> ReadList<T>(this BinaryReader Reader) where T : IBinarySerializable
 		{
-			T[] Items = Reader.ReadArray<T>(ReadElement);
-			return (Items == null)? null : new List<T>(Items);
+			return ReadList<T>(Reader, () => Reader.ReadObject<T>());
+		}
+
+		/// <summary>
+		/// Read a list of objects from a binary reader
+		/// </summary>
+		/// <typeparam name="T">The element type for the list</typeparam>
+		/// <param name="Reader">Reader to read data from</param>
+		/// <param name="ReadElement">Delegate to call to serialize each element</param>
+		/// <returns>List of objects, as serialized. May be null.</returns>
+		public static List<T> ReadList<T>(this BinaryReader Reader, Func<T> ReadElement)
+		{
+			int NumItems = Reader.ReadInt32();
+			if(NumItems < 0)
+			{
+				return null;
+			}
+
+			List<T> Items = new List<T>(NumItems);
+			for(int Idx = 0; Idx < NumItems; Idx++)
+			{
+				Items.Add(ReadElement());
+			}
+			return Items;
+		}
+
+		/// <summary>
+		/// Read a list of objects from a binary reader
+		/// </summary>
+		/// <typeparam name="T">The element type for the list</typeparam>
+		/// <param name="Reader">Reader to read data from</param>
+		/// <param name="ReadKey">Delegate to call to serialize each key</param>
+		/// <param name="ReadValue">Delegate to call to serialize each value</param>
+		/// <returns>List of objects, as serialized. May be null.</returns>
+		public static Dictionary<K, V> ReadDictionary<K, V>(this BinaryReader Reader, Func<K> ReadKey, Func<V> ReadValue)
+		{
+			int NumItems = Reader.ReadInt32();
+			if(NumItems < 0)
+			{
+				return null;
+			}
+
+			Dictionary<K, V> Items = new Dictionary<K, V>(NumItems);
+			for(int Idx = 0; Idx < NumItems; Idx++)
+			{
+				K Key = ReadKey();
+				V Value = ReadValue();
+				Items.Add(Key, Value);
+			}
+			return Items;
 		}
 
 		/// <summary>
@@ -77,7 +167,7 @@ namespace Tools.DotNETCommon
 			}
 			else if(ObjectType == typeof(string[]))
 			{
-				return Reader.ReadArray(x => x.ReadString());
+				return Reader.ReadStringArray();
 			}
 			else if(ObjectType == typeof(bool?))
 			{
diff --git a/Engine/Source/Programs/DotNETCommon/DotNETUtilities/BinaryWriterExtensions.cs b/Engine/Source/Programs/DotNETCommon/DotNETUtilities/BinaryWriterExtensions.cs
--- a/Engine/Source/Programs/DotNETCommon/DotNETUtilities/BinaryWriterExtensions.cs
+++ b/Engine/Source/Programs/DotNETCommon/DotNETUtilities/BinaryWriterExtensions.cs
@@ -9,8 +9,52 @@ using System.Threading.Tasks;
 
 namespace Tools.DotNETCommon
 {
+	/// <summary>
+	/// Interface that can be implemented by objects to support BinaryWriter container functionality. Also implies that
+	/// derived objects will have a constructor taking a single BinaryReader argument (similar to how the system ISerializable
+	/// interface assumes a specific constructor)
+	/// </summary>
+	public interface IBinarySerializable
+	{
+		void Write(BinaryWriter Writer);
+	}
+
+	/// <summary>
+	/// Extension methods for BinaryWriter class
+	/// </summary>
 	public static class BinaryWriterExtensions
 	{
+		/// <summary>
+		/// Writes an item implementing the IBinarySerializable interface to a binary writer. Included for symmetry with standard Writer.Write(X) calls.
+		/// </summary>
+		/// <param name="Writer">Writer to serialize to</param>
+		/// <param name="Item">The item to write</param>
+		public static void Write(this BinaryWriter Writer, IBinarySerializable Item)
+		{
+			Item.Write(Writer);
+		}
+
+		/// <summary>
+		/// Writes an array of strings to a binary writer.
+		/// </summary>
+		/// <param name="Writer">Writer to serialize to</param>
+		/// <param name="Items">Array of items</param>
+		public static void Write(this BinaryWriter Writer, string[] Items)
+		{
+			Write(Writer, Items, Item => Writer.Write(Item));
+		}
+
+		/// <summary>
+		/// Writes an array to a binary writer.
+		/// </summary>
+		/// <typeparam name="T">The array element type</typeparam>
+		/// <param name="Writer">Writer to serialize to</param>
+		/// <param name="Items">Array of items</param>
+		public static void Write<T>(this BinaryWriter Writer, T[] Items) where T : IBinarySerializable
+		{
+			Write(Writer, Items, Item => Item.Write(Writer));
+		}
+
 		/// <summary>
 		/// Writes an array to a binary writer.
 		/// </summary>
@@ -18,7 +62,7 @@ namespace Tools.DotNETCommon
 		/// <param name="Writer">Writer to serialize to</param>
 		/// <param name="Items">Array of items</param>
 		/// <param name="WriteElement">Delegate to call to serialize each element</param>
-		public static void Write<T>(this BinaryWriter Writer, T[] Items, Action<BinaryWriter, T> WriteElement)
+		public static void Write<T>(this BinaryWriter Writer, T[] Items, Action<T> WriteElement)
 		{
 			if(Items == null)
 			{
@@ -29,11 +73,32 @@ namespace Tools.DotNETCommon
 				Writer.Write(Items.Length);
 				for(int Idx = 0; Idx < Items.Length; Idx++)
 				{
-					WriteElement(Writer, Items[Idx]);
+					WriteElement(Items[Idx]);
 				}
 			}
 		}
 
+		/// <summary>
+		/// Writes a list of strings to a binary writer.
+		/// </summary>
+		/// <param name="Writer">Writer to serialize to</param>
+		/// <param name="Items">Array of items</param>
+		public static void Write<T>(this BinaryWriter Writer, List<string> Items)
+		{
+			Write(Writer, Items, Item => Writer.Write(Item));
+		}
+
+		/// <summary>
+		/// Writes a list to a binary writer.
+		/// </summary>
+		/// <typeparam name="T">The array element type</typeparam>
+		/// <param name="Writer">Writer to serialize to</param>
+		/// <param name="Items">Array of items</param>
+		public static void Write<T>(this BinaryWriter Writer, List<T> Items) where T : IBinarySerializable
+		{
+			Write(Writer, Items, Item => Item.Write(Writer));
+		}
+
 		/// <summary>
 		/// Writes a list to a binary writer.
 		/// </summary>
@@ -41,7 +106,7 @@ namespace Tools.DotNETCommon
 		/// <param name="Writer">Writer to serialize to</param>
 		/// <param name="Items">List of items</param>
 		/// <param name="WriteElement">Delegate to call to serialize each element</param>
-		public static void Write<T>(this BinaryWriter Writer, List<T> Items, Action<BinaryWriter, T> WriteElement)
+		public static void Write<T>(this BinaryWriter Writer, List<T> Items, Action<T> WriteElement)
 		{
 			if (Items == null)
 			{
@@ -52,7 +117,34 @@ namespace Tools.DotNETCommon
 				Writer.Write(Items.Count);
 				for (int Idx = 0; Idx < Items.Count; Idx++)
 				{
-					WriteElement(Writer, Items[Idx]);
+					WriteElement(Items[Idx]);
+				}
+			}
+		}
+
+		/// <summary>
+		/// Write a dictionary to a binary writer
+		/// </summary>
+		/// <typeparam name="K">The key type for the dictionary</typeparam>
+		/// <typeparam name="V">The value type for the dictionary</typeparam>
+		/// <param name="Reader">Reader to read data from</param>
+		/// <param name="Items">List of items to be written</param>
+		/// <param name="WriteKey">Delegate to call to serialize each key</param>
+		/// <param name="WriteKey">Delegate to call to serialize each value</param>
+		/// <returns>Dictionary of objects, as serialized. May be null.</returns>
+		public static void Write<K, V>(this BinaryWriter Writer, Dictionary<K, V> Items, Action<K> WriteKey, Action<V> WriteValue)
+		{
+			if(Items.Count == 0)
+			{
+				Writer.Write(-1);
+			}
+			else
+			{
+				Writer.Write(Items.Count);
+				foreach(KeyValuePair<K, V> Item in Items)
+				{
+					WriteKey(Item.Key);
+					WriteValue(Item.Value);
 				}
 			}
 		}
@@ -91,7 +183,7 @@ namespace Tools.DotNETCommon
 			}
 			else if(FieldType == typeof(string[]))
 			{
-				Writer.Write((string[])Value, (x, y) => x.Write(y));
+				Writer.Write((string[])Value);
 			}
 			else if(FieldType == typeof(bool?))
 			{
diff --git a/Engine/Source/Programs/UnrealBuildTool/Configuration/UEBuildDeployTarget.cs b/Engine/Source/Programs/UnrealBuildTool/Configuration/UEBuildDeployTarget.cs
--- a/Engine/Source/Programs/UnrealBuildTool/Configuration/UEBuildDeployTarget.cs
+++ b/Engine/Source/Programs/UnrealBuildTool/Configuration/UEBuildDeployTarget.cs
@@ -143,14 +143,14 @@ namespace UnrealBuildTool
 				TargetType = (TargetType)Reader.ReadInt32();
 				Platform = (UnrealTargetPlatform)Reader.ReadInt32();
 				Configuration = (UnrealTargetConfiguration)Reader.ReadInt32();
-				OutputPaths = Reader.ReadList(x => x.ReadFileReference());
+				OutputPaths = Reader.ReadList(() => Reader.ReadFileReference());
 				EngineIntermediateDirectory = Reader.ReadDirectoryReference();
 				ProjectDirectory = Reader.ReadDirectoryReference();
 				BuildReceiptFileName = Reader.ReadFileReference();
 				bCreateStubIPA = Reader.ReadBoolean();
-				AndroidArchitectures = Reader.ReadArray(x => x.ReadString());
-				AndroidGPUArchitectures = Reader.ReadArray(x => x.ReadString());
-				LuminGPUArchitectures = Reader.ReadArray(x => x.ReadString());
+				AndroidArchitectures = Reader.ReadStringArray();
+				AndroidGPUArchitectures = Reader.ReadStringArray();
+				LuminGPUArchitectures = Reader.ReadStringArray();
 			}
 		}
 
@@ -170,14 +170,14 @@ namespace UnrealBuildTool
 				Writer.Write((Int32)TargetType);
 				Writer.Write((Int32)Platform);
 				Writer.Write((Int32)Configuration);
-				Writer.Write(OutputPaths, (x, i) => x.Write(i));
+				Writer.Write(OutputPaths, Item => Writer.Write(Item));
 				Writer.Write(EngineIntermediateDirectory);
 				Writer.Write(ProjectDirectory);
 				Writer.Write(BuildReceiptFileName);
 				Writer.Write(bCreateStubIPA);
-				Writer.Write(AndroidArchitectures, (w, e) => w.Write(e));
-				Writer.Write(AndroidGPUArchitectures, (w, e) => w.Write(e));
-				Writer.Write(LuminGPUArchitectures, (w, e) => w.Write(e));
+				Writer.Write(AndroidArchitectures);
+				Writer.Write(AndroidGPUArchitectures);
+				Writer.Write(LuminGPUArchitectures);
 			}
 		}
 	}
diff --git a/Engine/Source/Programs/UnrealBuildTool/System/XmlConfigData.cs b/Engine/Source/Programs/UnrealBuildTool/System/XmlConfigData.cs
--- a/Engine/Source/Programs/UnrealBuildTool/System/XmlConfigData.cs
+++ b/Engine/Source/Programs/UnrealBuildTool/System/XmlConfigData.cs
@@ -70,7 +70,7 @@ namespace UnrealBuildTool
 				}
 
 				// Read the input files
-				FileReference[] InputFiles = Reader.ReadArray(x => x.ReadFileReference());
+				FileReference[] InputFiles = Reader.ReadArray(() => Reader.ReadFileReference());
 
 				// Read the types
 				int NumTypes = Reader.ReadInt32();
@@ -129,7 +129,7 @@ namespace UnrealBuildTool
 				Writer.Write(SerializationVersion);
 
 				// Save all the input files. The cache will not be valid if these change.
-				Writer.Write(InputFiles, (x, y) => x.Write(y));
+				Writer.Write(InputFiles, Item => Writer.Write(Item));
 
 				// Write all the categories
 				Writer.Write(TypeToValues.Count);
