Move the Log class from UBT to DotNetUtilities.

diff --git a/Engine/Plugins/ScriptPlugin/Source/ScriptGeneratorPlugin/ScriptGeneratorPlugin.Build.cs b/Engine/Plugins/ScriptPlugin/Source/ScriptGeneratorPlugin/ScriptGeneratorPlugin.Build.cs
--- a/Engine/Plugins/ScriptPlugin/Source/ScriptGeneratorPlugin/ScriptGeneratorPlugin.Build.cs
+++ b/Engine/Plugins/ScriptPlugin/Source/ScriptGeneratorPlugin/ScriptGeneratorPlugin.Build.cs
@@ -52,12 +52,12 @@ namespace UnrealBuildTool.Rules
 			var LuaLibPath = Path.Combine(LuaLibDirectory, "Lua.lib");
 			if (File.Exists(LuaLibPath))
 			{
-				Log.TraceVerbose("ScriptGenerator LUA Integration enabled");
+//				Log.TraceVerbose("ScriptGenerator LUA Integration enabled");
 				PublicDefinitions.Add("WITH_LUA=1");
 			}
 			else
 			{
-				Log.TraceVerbose("ScriptGenerator LUA Integration NOT enabled");
+//				Log.TraceVerbose("ScriptGenerator LUA Integration NOT enabled");
 				PublicDefinitions.Add("WITH_LUA=0");
 			}
 
diff --git a/Engine/Plugins/ScriptPlugin/Source/ScriptPlugin/ScriptPlugin.Build.cs b/Engine/Plugins/ScriptPlugin/Source/ScriptPlugin/ScriptPlugin.Build.cs
--- a/Engine/Plugins/ScriptPlugin/Source/ScriptPlugin/ScriptPlugin.Build.cs
+++ b/Engine/Plugins/ScriptPlugin/Source/ScriptPlugin/ScriptPlugin.Build.cs
@@ -68,11 +68,11 @@ namespace UnrealBuildTool.Rules
 				PublicLibraryPaths.Add(LuaLibDirectory);
 				PublicAdditionalLibraries.Add(LuaLibPath);
 
-				Log.TraceVerbose("LUA Integration enabled: {0}", IncludePath);
+//				Log.TraceVerbose("LUA Integration enabled: {0}", IncludePath);
 			}
 			else
 			{
-				Log.TraceVerbose("LUA Integration NOT enabled");
+//				Log.TraceVerbose("LUA Integration NOT enabled");
 				PublicDefinitions.Add("WITH_LUA=0");
 			}
 		}
diff --git a/Engine/Source/Programs/AutomationTool/AutomationUtils/Automation.cs b/Engine/Source/Programs/AutomationTool/AutomationUtils/Automation.cs
--- a/Engine/Source/Programs/AutomationTool/AutomationUtils/Automation.cs
+++ b/Engine/Source/Programs/AutomationTool/AutomationUtils/Automation.cs
@@ -1,4 +1,4 @@
-﻿// Copyright 1998-2018 Epic Games, Inc. All Rights Reserved.
+// Copyright 1998-2018 Epic Games, Inc. All Rights Reserved.
 using System;
 using System.Collections.Generic;
 using System.Linq;
@@ -7,6 +7,7 @@ using System.IO;
 using System.Reflection;
 using System.Diagnostics;
 using UnrealBuildTool;
+using Tools.DotNETCommon;
 
 namespace AutomationTool
 {
@@ -141,7 +142,7 @@ namespace AutomationTool
 	[Help(
 @"Executes scripted commands
 
-AutomationTool.exe [-verbose] [-compileonly] [-p4] Command0 [-Arg0 -Arg1 -Arg2 …] Command1 [-Arg0 -Arg1 …] Command2 [-Arg0 …] Commandn … [EnvVar0=MyValue0 … EnvVarn=MyValuen]"
+AutomationTool.exe [-verbose] [-compileonly] [-p4] Command0 [-Arg0 -Arg1 -Arg2 ...] Command1 [-Arg0 -Arg1 ...] Command2 [-Arg0 ...] Commandn ... [EnvVar0=MyValue0 ... EnvVarn=MyValuen]"
 )]
 	[Help("verbose", "Enables verbose logging")]
 	[Help("nop4", "Disables Perforce functionality (default if not run on a build machine)")]
diff --git a/Engine/Source/Programs/AutomationTool/AutomationUtils/CommandUtils.cs b/Engine/Source/Programs/AutomationTool/AutomationUtils/CommandUtils.cs
--- a/Engine/Source/Programs/AutomationTool/AutomationUtils/CommandUtils.cs
+++ b/Engine/Source/Programs/AutomationTool/AutomationUtils/CommandUtils.cs
@@ -116,7 +116,7 @@ namespace AutomationTool
 		[MethodImplAttribute(MethodImplOptions.NoInlining)]
 		public static void Log(string Format, params object[] Args)
 		{
-			UnrealBuildTool.Log.WriteLine(1, UnrealBuildTool.LogEventType.Console, Format, Args);
+			Tools.DotNETCommon.Log.WriteLine(1, Tools.DotNETCommon.LogEventType.Console, Format, Args);
 		}
 
 		/// <summary>
@@ -126,7 +126,7 @@ namespace AutomationTool
 		[MethodImplAttribute(MethodImplOptions.NoInlining)]
 		public static void Log(string Message)
 		{
-			UnrealBuildTool.Log.WriteLine(1, UnrealBuildTool.LogEventType.Console, Message);
+			Tools.DotNETCommon.Log.WriteLine(1, Tools.DotNETCommon.LogEventType.Console, Message);
 		}
 
 		/// <summary>
@@ -137,7 +137,7 @@ namespace AutomationTool
 		[MethodImplAttribute(MethodImplOptions.NoInlining)]
 		public static void LogError(string Format, params object[] Args)
 		{
-			UnrealBuildTool.Log.WriteLine(1, UnrealBuildTool.LogEventType.Error, Format, Args);
+			Tools.DotNETCommon.Log.WriteLine(1, Tools.DotNETCommon.LogEventType.Error, Format, Args);
 		}
 
 		/// <summary>
@@ -147,7 +147,7 @@ namespace AutomationTool
 		[MethodImplAttribute(MethodImplOptions.NoInlining)]
 		public static void LogError(string Message)
 		{
-			UnrealBuildTool.Log.WriteLine(1, UnrealBuildTool.LogEventType.Error, Message);
+			Tools.DotNETCommon.Log.WriteLine(1, Tools.DotNETCommon.LogEventType.Error, Message);
 		}
 
 		/// <summary>
@@ -158,7 +158,7 @@ namespace AutomationTool
 		[MethodImplAttribute(MethodImplOptions.NoInlining)]
 		public static void LogWarning(string Format, params object[] Args)
 		{
-			UnrealBuildTool.Log.WriteLine(1, UnrealBuildTool.LogEventType.Warning, Format, Args);
+			Tools.DotNETCommon.Log.WriteLine(1, Tools.DotNETCommon.LogEventType.Warning, Format, Args);
 		}
 
 		/// <summary>
@@ -168,7 +168,7 @@ namespace AutomationTool
 		[MethodImplAttribute(MethodImplOptions.NoInlining)]
 		public static void LogWarning(string Message)
 		{
-			UnrealBuildTool.Log.WriteLine(1, UnrealBuildTool.LogEventType.Warning, Message);
+			Tools.DotNETCommon.Log.WriteLine(1, Tools.DotNETCommon.LogEventType.Warning, Message);
 		}
 
 		/// <summary>
@@ -179,7 +179,7 @@ namespace AutomationTool
 		[MethodImplAttribute(MethodImplOptions.NoInlining)]
 		public static void LogVerbose(string Format, params object[] Args)
 		{
-			UnrealBuildTool.Log.WriteLine(1, UnrealBuildTool.LogEventType.Verbose, Format, Args);
+			Tools.DotNETCommon.Log.WriteLine(1, Tools.DotNETCommon.LogEventType.Verbose, Format, Args);
 		}
 
 		/// <summary>
@@ -189,7 +189,7 @@ namespace AutomationTool
 		[MethodImplAttribute(MethodImplOptions.NoInlining)]
 		public static void LogVerbose(string Message)
 		{
-			UnrealBuildTool.Log.WriteLine(1, UnrealBuildTool.LogEventType.Verbose, Message);
+			Tools.DotNETCommon.Log.WriteLine(1, Tools.DotNETCommon.LogEventType.Verbose, Message);
 		}
 
 		/// <summary>
@@ -200,7 +200,7 @@ namespace AutomationTool
 		[MethodImplAttribute(MethodImplOptions.NoInlining)]
 		public static void LogVeryVerbose(string Format, params object[] Args)
 		{
-			UnrealBuildTool.Log.WriteLine(1, UnrealBuildTool.LogEventType.VeryVerbose, Format, Args);
+			Tools.DotNETCommon.Log.WriteLine(1, Tools.DotNETCommon.LogEventType.VeryVerbose, Format, Args);
 		}
 
 		/// <summary>
@@ -210,7 +210,7 @@ namespace AutomationTool
 		[MethodImplAttribute(MethodImplOptions.NoInlining)]
 		public static void LogVeryVerbose(string Message)
 		{
-			UnrealBuildTool.Log.WriteLine(1, UnrealBuildTool.LogEventType.VeryVerbose, Message);
+			Tools.DotNETCommon.Log.WriteLine(1, Tools.DotNETCommon.LogEventType.VeryVerbose, Message);
 		}
 
 		/// <summary>
@@ -221,7 +221,7 @@ namespace AutomationTool
 		[MethodImplAttribute(MethodImplOptions.NoInlining)]
 		public static void LogLog(string Format, params object[] Args)
 		{
-			UnrealBuildTool.Log.WriteLine(1, UnrealBuildTool.LogEventType.Log, Format, Args);
+			Tools.DotNETCommon.Log.WriteLine(1, Tools.DotNETCommon.LogEventType.Log, Format, Args);
 		}
 
 		/// <summary>
@@ -231,7 +231,7 @@ namespace AutomationTool
 		[MethodImplAttribute(MethodImplOptions.NoInlining)]
 		public static void LogLog(string Message)
 		{
-			UnrealBuildTool.Log.WriteLine(1, UnrealBuildTool.LogEventType.Log, Message);
+			Tools.DotNETCommon.Log.WriteLine(1, LogEventType.Log, Message);
 		}
 
 		/// <summary>
@@ -241,9 +241,9 @@ namespace AutomationTool
 		/// <param name="Format">Format string</param>
 		/// <param name="Args">Arguments</param>
 		[MethodImplAttribute(MethodImplOptions.NoInlining)]
-		public static void LogWithVerbosity(UnrealBuildTool.LogEventType Verbosity, string Format, params object[] Args)
+		public static void LogWithVerbosity(Tools.DotNETCommon.LogEventType Verbosity, string Format, params object[] Args)
 		{
-            UnrealBuildTool.Log.WriteLine(1, Verbosity, Format, Args);
+            Tools.DotNETCommon.Log.WriteLine(1, Verbosity, Format, Args);
 		}
 
 		/// <summary>
@@ -252,9 +252,9 @@ namespace AutomationTool
 		/// <param name="Verbosity">Verbosity</param>
 		/// <param name="Message">Text</param>
 		[MethodImplAttribute(MethodImplOptions.NoInlining)]
-		public static void LogWithVerbosity(UnrealBuildTool.LogEventType Verbosity, string Message)
+		public static void LogWithVerbosity(Tools.DotNETCommon.LogEventType Verbosity, string Message)
 		{
-            UnrealBuildTool.Log.WriteLine(1, Verbosity, Message);
+            Tools.DotNETCommon.Log.WriteLine(1, Verbosity, Message);
 		}
 
 		/// <summary>
@@ -263,9 +263,9 @@ namespace AutomationTool
 		/// <param name="Verbosity">Verbosity</param>
 		/// <param name="Ex">Exception</param>
 		[MethodImplAttribute(MethodImplOptions.NoInlining)]
-		public static void LogWithVerbosity(UnrealBuildTool.LogEventType Verbosity, Exception Ex)
+		public static void LogWithVerbosity(Tools.DotNETCommon.LogEventType Verbosity, Exception Ex)
 		{
-            UnrealBuildTool.Log.WriteLine(1, Verbosity, LogUtils.FormatException(Ex));
+            Tools.DotNETCommon.Log.WriteLine(1, Verbosity, LogUtils.FormatException(Ex));
 		}
 
 		#endregion
diff --git a/Engine/Source/Programs/AutomationTool/AutomationUtils/P4Environment.cs b/Engine/Source/Programs/AutomationTool/AutomationUtils/P4Environment.cs
--- a/Engine/Source/Programs/AutomationTool/AutomationUtils/P4Environment.cs
+++ b/Engine/Source/Programs/AutomationTool/AutomationUtils/P4Environment.cs
@@ -1,10 +1,11 @@
-﻿// Copyright 1998-2018 Epic Games, Inc. All Rights Reserved.
+// Copyright 1998-2018 Epic Games, Inc. All Rights Reserved.
 
 using System;
 using System.Collections.Generic;
 using System.Linq;
 using System.Text;
 using UnrealBuildTool;
+using Tools.DotNETCommon;
 
 namespace AutomationTool
 {
diff --git a/Engine/Source/Programs/AutomationTool/AutomationUtils/ParallelExecutor.cs b/Engine/Source/Programs/AutomationTool/AutomationUtils/ParallelExecutor.cs
--- a/Engine/Source/Programs/AutomationTool/AutomationUtils/ParallelExecutor.cs
+++ b/Engine/Source/Programs/AutomationTool/AutomationUtils/ParallelExecutor.cs
@@ -8,6 +8,7 @@ using System.Text;
 using System.Threading;
 using System.Threading.Tasks;
 using System.Xml;
+using Tools.DotNETCommon;
 
 namespace AutomationTool
 {
@@ -84,7 +85,7 @@ namespace AutomationTool
 
 			CommandUtils.Log("Building with {0} {1}...", MaxProcesses, (MaxProcesses == 1)? "process" : "processes");
 
-			using (UnrealBuildTool.ScopedLogIndent Indent = new UnrealBuildTool.ScopedLogIndent("  "))
+			using (ScopedLogIndent Indent = new ScopedLogIndent("  "))
 			{
 				// Create the list of things to process
 				List<BuildAction> QueuedActions = new List<BuildAction>();
diff --git a/Engine/Source/Programs/AutomationTool/AutomationUtils/ProcessUtils.cs b/Engine/Source/Programs/AutomationTool/AutomationUtils/ProcessUtils.cs
--- a/Engine/Source/Programs/AutomationTool/AutomationUtils/ProcessUtils.cs
+++ b/Engine/Source/Programs/AutomationTool/AutomationUtils/ProcessUtils.cs
@@ -9,6 +9,7 @@ using System.Diagnostics;
 using System.Management;
 using System.Runtime.InteropServices;
 using System.Runtime.CompilerServices;
+using Tools.DotNETCommon;
 
 namespace AutomationTool
 {
@@ -53,7 +54,7 @@ namespace AutomationTool
 		/// Creates a new process and adds it to the tracking list.
 		/// </summary>
 		/// <returns>New Process objects</returns>
-		public static IProcessResult CreateProcess(string AppName, bool bAllowSpew, bool bCaptureSpew, Dictionary<string, string> Env = null, UnrealBuildTool.LogEventType SpewVerbosity = UnrealBuildTool.LogEventType.Console, ProcessResult.SpewFilterCallbackType SpewFilterCallback = null)
+		public static IProcessResult CreateProcess(string AppName, bool bAllowSpew, bool bCaptureSpew, Dictionary<string, string> Env = null, LogEventType SpewVerbosity = LogEventType.Console, ProcessResult.SpewFilterCallbackType SpewFilterCallback = null)
 		{
 			var NewProcess = HostPlatform.Current.CreateProcess(AppName);
 			if (Env != null)
@@ -230,7 +231,7 @@ namespace AutomationTool
 		private int ProcessExitCode = -1;
 		private StringBuilder ProcessOutput;
 		private bool AllowSpew = true;
-		private UnrealBuildTool.LogEventType SpewVerbosity = UnrealBuildTool.LogEventType.Console;
+		private LogEventType SpewVerbosity = LogEventType.Console;
 		private SpewFilterCallbackType SpewFilterCallback = null;
 		private string AppName = String.Empty;
 		private Process Proc = null;
@@ -238,7 +239,7 @@ namespace AutomationTool
 		private AutoResetEvent ErrorWaitHandle = new AutoResetEvent(false);
 		private object ProcSyncObject;
 
-		public ProcessResult(string InAppName, Process InProc, bool bAllowSpew, bool bCaptureSpew = true, UnrealBuildTool.LogEventType SpewVerbosity = UnrealBuildTool.LogEventType.Console, SpewFilterCallbackType InSpewFilterCallback = null)
+		public ProcessResult(string InAppName, Process InProc, bool bAllowSpew, bool bCaptureSpew = true, LogEventType SpewVerbosity = LogEventType.Console, SpewFilterCallbackType InSpewFilterCallback = null)
 		{
 			AppName = InAppName;
 			ProcSyncObject = new object();
@@ -284,9 +285,9 @@ namespace AutomationTool
         /// <param name="Verbosity"></param>
         /// <param name="Message"></param>
         [MethodImplAttribute(MethodImplOptions.NoInlining)]
-		private void LogOutput(UnrealBuildTool.LogEventType Verbosity, string Message)
+		private void LogOutput(LogEventType Verbosity, string Message)
 		{
-            UnrealBuildTool.Log.WriteLine(1, Verbosity, Message);
+            Log.WriteLine(1, Verbosity, Message);
 		}
        
 		/// <summary>
@@ -805,14 +806,14 @@ namespace AutomationTool
 			}
 			var StartTime = DateTime.UtcNow;
 
-			UnrealBuildTool.LogEventType SpewVerbosity = Options.HasFlag(ERunOptions.SpewIsVerbose) ? UnrealBuildTool.LogEventType.Verbose : UnrealBuildTool.LogEventType.Console;
+			LogEventType SpewVerbosity = Options.HasFlag(ERunOptions.SpewIsVerbose) ? LogEventType.Verbose : LogEventType.Console;
             if (!Options.HasFlag(ERunOptions.NoLoggingOfRunCommand))
             {
                 LogWithVerbosity(SpewVerbosity,"Running: " + App + " " + (String.IsNullOrEmpty(CommandLine) ? "" : CommandLine));
             }
 
-			string PrevIndent = UnrealBuildTool.Log.Indent;
-			UnrealBuildTool.Log.Indent += "  ";
+			string PrevIndent = Tools.DotNETCommon.Log.Indent;
+			Tools.DotNETCommon.Log.Indent += "  ";
 
 			IProcessResult Result = ProcessManager.CreateProcess(App, Options.HasFlag(ERunOptions.AllowSpew), !Options.HasFlag(ERunOptions.NoStdOutCapture), Env, SpewVerbosity: SpewVerbosity, SpewFilterCallback: SpewFilterCallback);
 			try
@@ -861,7 +862,7 @@ namespace AutomationTool
 			}
 			finally
 			{
-				UnrealBuildTool.Log.Indent = PrevIndent;
+				Tools.DotNETCommon.Log.Indent = PrevIndent;
 			}
 
 			if (!Options.HasFlag(ERunOptions.NoWaitForExit))
diff --git a/Engine/Source/Programs/AutomationTool/BuildGraph/Script.cs b/Engine/Source/Programs/AutomationTool/BuildGraph/Script.cs
--- a/Engine/Source/Programs/AutomationTool/BuildGraph/Script.cs
+++ b/Engine/Source/Programs/AutomationTool/BuildGraph/Script.cs
@@ -81,7 +81,7 @@ namespace AutomationTool
 				{
 					if (!Document.bHasErrors)
 					{
-						UnrealBuildTool.Log.WriteLine(LogEventType.Error, LogFormatOptions.NoSeverityPrefix, "{0}({1}): error: {2}", File.FullName, Ex.LineNumber, Ex.Message);
+						Log.WriteLine(LogEventType.Error, LogFormatOptions.NoSeverityPrefix, "{0}({1}): error: {2}", File.FullName, Ex.LineNumber, Ex.Message);
 						Document.bHasErrors = true;
 					}
 				}
diff --git a/Engine/Source/Programs/AutomationTool/Program.cs b/Engine/Source/Programs/AutomationTool/Program.cs
--- a/Engine/Source/Programs/AutomationTool/Program.cs
+++ b/Engine/Source/Programs/AutomationTool/Program.cs
@@ -1,4 +1,4 @@
-﻿// Copyright 1998-2018 Epic Games, Inc. All Rights Reserved.
+// Copyright 1998-2018 Epic Games, Inc. All Rights Reserved.
 // This software is provided "as-is," without any express or implied warranty. 
 // In no event shall the author, nor Epic Games, Inc. be held liable for any damages arising from the use of this software.
 // This software will not be supported.
@@ -35,7 +35,7 @@ namespace AutomationTool
             }
 
 			// Parse the log level argument
-			UnrealBuildTool.LogEventType LogLevel = LogEventType.Log;
+			LogEventType LogLevel = LogEventType.Log;
 			if(CommandUtils.ParseParam(Arguments, "-Verbose"))
 			{
 				LogLevel = LogEventType.Verbose;
@@ -47,7 +47,7 @@ namespace AutomationTool
 
 			// Initialize the log system, buffering the output until we can create the log file
 			StartupTraceListener StartupListener = new StartupTraceListener();
-			UnrealBuildTool.Log.InitLogging(
+			Log.InitLogging(
                 bLogTimestamps: CommandUtils.ParseParam(Arguments, "-Timestamps"),
 				InLogLevel: LogLevel,
                 bLogSeverity: true,
@@ -57,7 +57,7 @@ namespace AutomationTool
                 bColorConsoleOutput: true,
                 TraceListeners: new TraceListener[]
                 {
-                    new UEConsoleTraceListener(),
+                    new Tools.DotNETCommon.ConsoleTraceListener(),
 					StartupListener
 				}
 			);
diff --git a/Engine/Source/Programs/AutomationTool/Scripts/BuildPhysX.Automation.cs b/Engine/Source/Programs/AutomationTool/Scripts/BuildPhysX.Automation.cs
--- a/Engine/Source/Programs/AutomationTool/Scripts/BuildPhysX.Automation.cs
+++ b/Engine/Source/Programs/AutomationTool/Scripts/BuildPhysX.Automation.cs
@@ -572,7 +572,7 @@ class BuildPhysX : BuildCommand
 	{
 		Process LocalProcess = new Process();
 		LocalProcess.StartInfo = StartInfo;
-		LocalProcess.OutputDataReceived += (Sender, Line) => { if (Line != null && Line.Data != null) UnrealBuildTool.Log.TraceInformation(Line.Data); };
+		LocalProcess.OutputDataReceived += (Sender, Line) => { if (Line != null && Line.Data != null) Tools.DotNETCommon.Log.TraceInformation(Line.Data); };
 		return RunLocalProcess(LocalProcess);
 	}
 
diff --git a/Engine/Source/Programs/AutomationTool/Scripts/CheckForHacks.cs b/Engine/Source/Programs/AutomationTool/Scripts/CheckForHacks.cs
--- a/Engine/Source/Programs/AutomationTool/Scripts/CheckForHacks.cs
+++ b/Engine/Source/Programs/AutomationTool/Scripts/CheckForHacks.cs
@@ -153,7 +153,7 @@ namespace AutomationTool
 								{
 									lock (this)
 									{
-										UnrealBuildTool.Log.WriteLine(UnrealBuildTool.LogEventType.Warning, UnrealBuildTool.LogFormatOptions.NoSeverityPrefix, "{0}({1}): warning: Code should not be in this branch: '{2}'", FileToCheck.FullName, LineNumber, Line.Trim());
+										Tools.DotNETCommon.Log.WriteLine(Tools.DotNETCommon.LogEventType.Warning, Tools.DotNETCommon.LogFormatOptions.NoSeverityPrefix, "{0}({1}): warning: Code should not be in this branch: '{2}'", FileToCheck.FullName, LineNumber, Line.Trim());
 									}
 								}
 							}
diff --git a/Engine/Source/Programs/AutomationTool/Scripts/RebuildHLODCommand.Automation.cs b/Engine/Source/Programs/AutomationTool/Scripts/RebuildHLODCommand.Automation.cs
--- a/Engine/Source/Programs/AutomationTool/Scripts/RebuildHLODCommand.Automation.cs
+++ b/Engine/Source/Programs/AutomationTool/Scripts/RebuildHLODCommand.Automation.cs
@@ -165,11 +165,11 @@ namespace AutomationScripts.Automation
                 if ( AEx != null )
                 {
                     string LogFile = AEx.LogFileName;
-                    UnrealBuildTool.Log.TraceWarning("Attempting to load file {0}", LogFile);
+                    Tools.DotNETCommon.Log.TraceWarning("Attempting to load file {0}", LogFile);
                     if ( LogFile != "")
                     {
                         
-                        UnrealBuildTool.Log.TraceWarning("Attempting to read file {0}", LogFile);
+                        Tools.DotNETCommon.Log.TraceWarning("Attempting to read file {0}", LogFile);
                         try
                         {
                             string[] AllLogFile = ReadAllLines(LogFile);
diff --git a/Engine/Source/Programs/AutomationTool/Scripts/RebuildLightMapsCommand.Automation.cs b/Engine/Source/Programs/AutomationTool/Scripts/RebuildLightMapsCommand.Automation.cs
--- a/Engine/Source/Programs/AutomationTool/Scripts/RebuildLightMapsCommand.Automation.cs
+++ b/Engine/Source/Programs/AutomationTool/Scripts/RebuildLightMapsCommand.Automation.cs
@@ -142,11 +142,11 @@ namespace AutomationScripts.Automation
                 if ( AEx != null )
                 {
                     string LogFile = AEx.LogFileName;
-                    UnrealBuildTool.Log.TraceWarning("Attempting to load file {0}", LogFile);
+                    Tools.DotNETCommon.Log.TraceWarning("Attempting to load file {0}", LogFile);
                     if ( LogFile != "")
                     {
                         
-                        UnrealBuildTool.Log.TraceWarning("Attempting to read file {0}", LogFile);
+                        Tools.DotNETCommon.Log.TraceWarning("Attempting to read file {0}", LogFile);
                         try
                         {
                             string[] AllLogFile = ReadAllLines(LogFile);
diff --git a/Engine/Source/Programs/AutomationTool/Scripts/ResavePackagesCommand.Automation.cs b/Engine/Source/Programs/AutomationTool/Scripts/ResavePackagesCommand.Automation.cs
--- a/Engine/Source/Programs/AutomationTool/Scripts/ResavePackagesCommand.Automation.cs
+++ b/Engine/Source/Programs/AutomationTool/Scripts/ResavePackagesCommand.Automation.cs
@@ -144,11 +144,11 @@ namespace AutomationScripts.Automation
                 if ( AEx != null )
                 {
                     string LogFile = AEx.LogFileName;
-                    UnrealBuildTool.Log.TraceWarning("Attempting to load file {0}", LogFile);
+                    Tools.DotNETCommon.Log.TraceWarning("Attempting to load file {0}", LogFile);
                     if ( LogFile != "")
                     {
                         
-                        UnrealBuildTool.Log.TraceWarning("Attempting to read file {0}", LogFile);
+                        Tools.DotNETCommon.Log.TraceWarning("Attempting to read file {0}", LogFile);
                         try
                         {
                             string[] AllLogFile = ReadAllLines(LogFile);
diff --git a/Engine/Source/Programs/DotNETCommon/DotNETUtilities/ConsoleTraceListener.cs b/Engine/Source/Programs/DotNETCommon/DotNETUtilities/ConsoleTraceListener.cs
new file mode 100644
--- /dev/null
+++ b/Engine/Source/Programs/DotNETCommon/DotNETUtilities/ConsoleTraceListener.cs
@@ -0,0 +1,51 @@
+﻿// Copyright 1998-2018 Epic Games, Inc. All Rights Reserved.
+
+using System;
+using System.Collections.Generic;
+using System.Diagnostics;
+using System.IO;
+using System.Linq;
+using System.Text;
+using System.Threading.Tasks;
+
+namespace Tools.DotNETCommon
+{
+	/// <summary>
+	/// Replacement class for System.Diagnostics.ConsoleTracerListener, which doesn't exist in .NET Core
+	/// </summary>
+	public class ConsoleTraceListener : TraceListener
+	{
+		/// <summary>
+		/// Destination stream writer
+		/// </summary>
+		private TextWriter Writer; 
+
+		/// <summary>
+		/// Constructor
+		/// </summary>
+		/// <param name="bUseStandardError">TRUE if we are writing to the standard error stream, FALSE for standard output</param>
+		public ConsoleTraceListener(bool bUseStandardError = false)
+		{
+			Writer = bUseStandardError ? System.Console.Error : System.Console.Out;
+		}
+
+		/// <summary>
+		/// Writes a string of text to the output stream
+		/// </summary>
+		/// <param name="message">Message to write</param>
+		public override void Write(string message)
+		{
+			Writer.Write(message);
+		}
+
+
+		/// <summary>
+		/// Writes a line of text to the output stream
+		/// </summary>
+		/// <param name="message">MEssage to write</param>
+		public override void WriteLine(string message)
+		{
+			Writer.WriteLine(message);
+		}
+	}
+}
diff --git a/Engine/Source/Programs/DotNETCommon/DotNETUtilities/DotNETUtilities.csproj b/Engine/Source/Programs/DotNETCommon/DotNETUtilities/DotNETUtilities.csproj
--- a/Engine/Source/Programs/DotNETCommon/DotNETUtilities/DotNETUtilities.csproj
+++ b/Engine/Source/Programs/DotNETCommon/DotNETUtilities/DotNETUtilities.csproj
@@ -57,6 +57,7 @@
     <Compile Include="AssemblyUtils.cs" />
     <Compile Include="BinaryReaderExtensions.cs" />
     <Compile Include="BinaryWriterExtensions.cs" />
+    <Compile Include="ConsoleTraceListener.cs" />
     <Compile Include="CsProjectInfo.cs" />
     <Compile Include="DirectoryReference.cs" />
     <Compile Include="FastJSON\Formatter.cs" />
@@ -73,6 +74,7 @@
     <Compile Include="FileSystemName.cs" />
     <Compile Include="FileSystemReference.cs" />
     <Compile Include="JsonObject.cs" />
+    <Compile Include="Log.cs" />
     <Compile Include="Properties\AssemblyInfo.cs" />
     <Compile Include="ResXResource\UEResXReader.cs" />
     <Compile Include="ResXResource\UEResXWriter.cs" />
diff --git a/Engine/Source/Programs/DotNETCommon/DotNETUtilities/Log.cs b/Engine/Source/Programs/DotNETCommon/DotNETUtilities/Log.cs
new file mode 100644
--- /dev/null
+++ b/Engine/Source/Programs/DotNETCommon/DotNETUtilities/Log.cs
@@ -0,0 +1,671 @@
+// Copyright 1998-2018 Epic Games, Inc. All Rights Reserved.
+
+using System;
+using System.Collections.Generic;
+using System.Diagnostics;
+using System.IO;
+using System.Linq;
+using System.Reflection;
+using System.Runtime.CompilerServices;
+using System.Text;
+using System.Threading.Tasks;
+
+namespace Tools.DotNETCommon
+{
+	/// <summary>
+	/// Log Event Type
+	/// </summary>
+	public enum LogEventType
+	{
+		/// <summary>
+		/// The log event is a fatal error
+		/// </summary>
+		Fatal = 1,
+
+		/// <summary>
+		/// The log event is an error
+		/// </summary>
+		Error = 2,
+
+		/// <summary>
+		/// The log event is a warning
+		/// </summary>
+		Warning = 4,
+
+		/// <summary>
+		/// Output the log event to the console
+		/// </summary>
+		Console = 8,
+
+		/// <summary>
+		/// Output the event to the on-disk log
+		/// </summary>
+		Log = 16,
+
+		/// <summary>
+		/// The log event should only be displayed if verbose logging is enabled
+		/// </summary>
+		Verbose = 32,
+
+		/// <summary>
+		/// The log event should only be displayed if very verbose logging is enabled
+		/// </summary>
+		VeryVerbose = 64
+	}
+
+	/// <summary>
+	/// Options for formatting messages
+	/// </summary>
+	[Flags]
+	public enum LogFormatOptions
+	{
+		/// <summary>
+		/// Format normally
+		/// </summary>
+		None = 0,
+
+		/// <summary>
+		/// Never write a severity prefix. Useful for pre-formatted messages that need to be in a particular format for, eg. the Visual Studio output window
+		/// </summary>
+		NoSeverityPrefix = 1,
+	}
+
+	/// <summary>
+	/// UAT/UBT Custom log system.
+	/// 
+	/// This lets you use any TraceListeners you want, but you should only call the static 
+	/// methods below, not call Trace.XXX directly, as the static methods
+	/// This allows the system to enforce the formatting and filtering conventions we desire.
+	///
+	/// For posterity, we cannot use the Trace or TraceSource class directly because of our special log requirements:
+	///   1. We possibly capture the method name of the logging event. This cannot be done as a macro, so must be done at the top level so we know how many layers of the stack to peel off to get the real function.
+	///   2. We have a verbose filter we would like to apply to all logs without having to have each listener filter individually, which would require our string formatting code to run every time.
+	///   3. We possibly want to ensure severity prefixes are logged, but Trace.WriteXXX does not allow any severity info to be passed down.
+	/// </summary>
+	static public class Log
+	{
+		/// <summary>
+		/// Guard our initialization. Mainly used by top level exception handlers to ensure its safe to call a logging function.
+		/// In general user code should not concern itself checking for this.
+		/// </summary>
+		private static bool bIsInitialized = false;
+
+		/// <summary>
+		/// When true, verbose logging is enabled.
+		/// </summary>
+		private static LogEventType LogLevel = LogEventType.VeryVerbose;
+
+		/// <summary>
+		/// When true, warnings and errors will have a WARNING: or ERROR: prexifx, respectively.
+		/// </summary>
+		private static bool bLogSeverity = false;
+		
+		/// <summary>
+		/// When true, warnings and errors will have a prefix suitable for display by MSBuild (avoiding error messages showing as (EXEC : Error : ")
+		/// </summary>
+		private static bool bLogProgramNameWithSeverity = false;
+		/// <summary>
+		/// When true, logs will have the calling mehod prepended to the output as MethodName:
+		/// </summary>
+		private static bool bLogSources = false;
+		
+		/// <summary>
+		/// When true, console output will have the calling mehod prepended to the output as MethodName:
+		/// </summary>
+		private static bool bLogSourcesToConsole = false;
+		
+		/// <summary>
+		/// When true, will detect warnings and errors and set the console output color to yellow and red.
+		/// </summary>
+		private static bool bColorConsoleOutput = false;
+		
+		/// <summary>
+		/// When configured, this tracks time since initialization to prepend a timestamp to each log.
+		/// </summary>
+		private static Stopwatch Timer;
+
+		/// <summary>
+		/// Expose the log level. This is a hack for ProcessResult.LogOutput, which wants to bypass our normal formatting scheme.
+		/// </summary>
+		public static bool bIsVerbose { get { return LogLevel >= LogEventType.Verbose; } }
+
+		/// <summary>
+		/// A collection of strings that have been already written once
+		/// </summary>
+		private static List<string> WriteOnceSet = new List<string>();
+
+		/// <summary>
+		/// Indent added to every output line
+		/// </summary>
+		public static string Indent
+		{
+			get;
+			set;
+		}
+
+		/// <summary>
+		/// Static initializer
+		/// </summary>
+		static Log()
+		{
+			Indent = "";
+		}
+
+		/// <summary>
+		/// Allows code to check if the log system is ready yet.
+		/// End users should NOT need to use this. It pretty much exists
+		/// to work around startup issues since this is a global singleton.
+		/// </summary>
+		/// <returns></returns>
+		public static bool IsInitialized()
+		{
+			return bIsInitialized;
+		}
+
+		/// <summary>
+		/// Allows code to check if the log system is using console output color.
+		/// </summary>
+		/// <returns></returns>
+		public static bool ColorConsoleOutput()
+		{
+			return bColorConsoleOutput;
+		}
+
+		/// <summary>
+		/// Allows us to change verbosity after initializing. This can happen since we initialize logging early, 
+		/// but then read the config and command line later, which could change this value.
+		/// </summary>
+		public static void SetLoggingLevel(LogEventType InLogLevel)
+		{
+			Log.LogLevel = InLogLevel;
+		}
+
+		/// <summary>
+		/// This class allows InitLogging to be called more than once to work around chicken and eggs issues with logging and parsing command lines (see UBT startup code).
+		/// </summary>
+		/// <param name="bLogTimestamps">If true, the timestamp from Log init time will be prepended to all logs.</param>
+		/// <param name="InLogLevel"></param>
+		/// <param name="bLogSeverity">If true, warnings and errors will have a WARNING: and ERROR: prefix to them. </param>
+		/// <param name="bLogProgramNameWithSeverity">If true, includes the program name with any severity prefix</param>
+		/// <param name="bLogSources">If true, logs will have the originating method name prepended to them.</param>
+		/// <param name="bLogSourcesToConsole">If true, console output will have the originating method name appended to it.</param>
+		/// <param name="bColorConsoleOutput"></param>
+		/// <param name="TraceListeners">Collection of trace listeners to attach to the Trace.Listeners, in addition to the Default listener. The existing listeners (except the Default listener) are cleared first.</param>
+		public static void InitLogging(bool bLogTimestamps, LogEventType InLogLevel, bool bLogSeverity, bool bLogProgramNameWithSeverity, bool bLogSources, bool bLogSourcesToConsole, bool bColorConsoleOutput, IEnumerable<TraceListener> TraceListeners)
+		{
+			bIsInitialized = true;
+			Timer = (bLogTimestamps && Timer == null) ? Stopwatch.StartNew() : null;
+			Log.LogLevel = InLogLevel;
+			Log.bLogSeverity = bLogSeverity;
+			Log.bLogProgramNameWithSeverity = bLogProgramNameWithSeverity;
+			Log.bLogSources = bLogSources;
+			Log.bLogSourcesToConsole = bLogSourcesToConsole;
+			Log.bColorConsoleOutput = bColorConsoleOutput;
+
+			// ensure that if InitLogging is called more than once we don't stack listeners.
+			// but always leave the default listener around.
+			for (int ListenerNdx = 0; ListenerNdx < Trace.Listeners.Count;)
+			{
+				if (Trace.Listeners[ListenerNdx].GetType() != typeof(DefaultTraceListener))
+				{
+					Trace.Listeners.RemoveAt(ListenerNdx);
+				}
+				else
+				{
+					++ListenerNdx;
+				}
+			}
+			// don't add any null listeners
+			Trace.Listeners.AddRange(TraceListeners.Where(l => l != null).ToArray());
+			Trace.AutoFlush = true;
+		}
+
+		/// <summary>
+		/// Gets the name of the Method N levels deep in the stack frame. Used to trap what method actually made the logging call.
+		/// Only used when bLogSources is true.
+		/// </summary>
+		/// <param name="StackFramesToSkip"></param>
+		/// <returns>ClassName.MethodName</returns>
+		[MethodImplAttribute(MethodImplOptions.NoInlining)]
+		private static string GetSource(int StackFramesToSkip)
+		{
+			StackFrame Frame = new StackFrame(2 + StackFramesToSkip);
+			System.Reflection.MethodBase Method = Frame.GetMethod();
+			return String.Format("{0}.{1}", Method.DeclaringType.Name, Method.Name);
+		}
+
+		/// <summary>
+		/// Converts a LogEventType into a log prefix. Only used when bLogSeverity is true.
+		/// </summary>
+		/// <param name="Severity"></param>
+		/// <returns></returns>
+		private static string GetSeverityPrefix(LogEventType Severity)
+		{
+			switch (Severity)
+			{
+				case LogEventType.Fatal:
+					return "FATAL ERROR: ";
+				case LogEventType.Error:
+					return "ERROR: ";
+				case LogEventType.Warning:
+					return "WARNING: ";
+				case LogEventType.Console:
+					return "";
+				case LogEventType.Verbose:
+					return "VERBOSE: ";
+				case LogEventType.VeryVerbose:
+					return "VVERBOSE: ";
+				default:
+					return "";
+			}
+		}
+
+		/// <summary>
+		/// Converts a LogEventType into a message code
+		/// </summary>
+		/// <param name="Severity"></param>
+		/// <returns></returns>
+		private static int GetMessageCode(LogEventType Severity)
+		{
+			return (int)Severity;
+		}
+
+		/// <summary>
+		/// Formats message for logging. Enforces the configured options.
+		/// </summary>
+		/// <param name="StackFramesToSkip">Number of frames to skip to get to the originator of the log request.</param>
+		/// <param name="Verbosity">Message verbosity level</param>
+		/// <param name="Options">Options for formatting this string</param>
+		/// <param name="bForConsole">Whether the message is intended for console output</param>
+		/// <param name="Format">Message text format string</param>
+		/// <param name="Args">Message text parameters</param>
+		/// <returns>Formatted message</returns>
+		[MethodImplAttribute(MethodImplOptions.NoInlining)]
+		private static List<string> FormatMessage(int StackFramesToSkip, LogEventType Verbosity, LogFormatOptions Options, bool bForConsole, string Format, params object[] Args)
+		{
+			string TimePrefix = (Timer != null) ? String.Format("[{0:hh\\:mm\\:ss\\.fff}] ", Timer.Elapsed) : "";
+			string SourcePrefix = (bForConsole ? bLogSourcesToConsole : bLogSources) ? string.Format("{0}: ", GetSource(StackFramesToSkip)) : "";
+			string SeverityPrefix = (bLogSeverity && ((Options & LogFormatOptions.NoSeverityPrefix) == 0)) ? GetSeverityPrefix(Verbosity) : "";
+
+			// Include the executable name when running inside MSBuild. If unspecified, MSBuild re-formats them with an "EXEC :" prefix.
+			if(SeverityPrefix.Length > 0 && bLogProgramNameWithSeverity)
+			{
+				SeverityPrefix = String.Format("{0}: {1}", Path.GetFileNameWithoutExtension(Assembly.GetEntryAssembly().Location), SeverityPrefix);
+			}
+
+			// If there are no extra args, don't try to format the string, in case it has any format control characters in it (our LOCTEXT strings tend to).
+			string[] Lines = ((Args.Length > 0) ? String.Format(Format, Args) : Format).TrimEnd(' ', '\t', '\r', '\n').Split('\n');
+
+			List<string> FormattedLines = new List<string>();
+			FormattedLines.Add(String.Format("{0}{1}{2}{3}{4}", TimePrefix, SourcePrefix, Indent, SeverityPrefix, Lines[0].TrimEnd('\r')));
+
+			if (Lines.Length > 1)
+			{
+				int PaddingLength = 0;
+				while(PaddingLength < Lines[0].Length && Char.IsWhiteSpace(Lines[0][PaddingLength]))
+				{
+					PaddingLength++;
+				}
+
+				string Padding = new string(' ', SeverityPrefix.Length) + Lines[0].Substring(0, PaddingLength);
+				for (int Idx = 1; Idx < Lines.Length; Idx++)
+				{
+					FormattedLines.Add(String.Format("{0}{1}{2}{3}{4}", TimePrefix, SourcePrefix, Indent, Padding, Lines[Idx].TrimEnd('\r')));
+				}
+			}
+
+			return FormattedLines;
+		}
+
+		/// <summary>
+		/// Writes a formatted message to the console. All other functions should boil down to calling this method.
+		/// </summary>
+		/// <param name="StackFramesToSkip">Number of frames to skip to get to the originator of the log request.</param>
+		/// <param name="bWriteOnce">If true, this message will be written only once</param>
+		/// <param name="Verbosity">Message verbosity level. We only meaningfully use values up to Verbose</param>
+		/// <param name="FormatOptions">Options for formatting messages</param>
+		/// <param name="Format">Message format string.</param>
+		/// <param name="Args">Optional arguments</param>
+		[MethodImplAttribute(MethodImplOptions.NoInlining)]
+		private static void WriteLinePrivate(int StackFramesToSkip, bool bWriteOnce, LogEventType Verbosity, LogFormatOptions FormatOptions, string Format, params object[] Args)
+		{
+			if (!bIsInitialized)
+			{
+				throw new InvalidOperationException("Tried to using Logging system before it was ready");
+			}
+
+			// if we want this message only written one time, check if it was already written out
+			if (bWriteOnce)
+			{
+				string Formatted = string.Format(Format, Args);
+				if (WriteOnceSet.Contains(Formatted))
+				{
+					return;
+				}
+
+				WriteOnceSet.Add(Formatted);
+			}
+
+			if (Verbosity <= LogLevel)
+			{
+				// Do console color highlighting here.
+				ConsoleColor DefaultColor = ConsoleColor.Gray;
+				bool bIsWarning = false;
+				bool bIsError = false;
+				// don't try to touch the console unless we are told to color the output.
+				if (bColorConsoleOutput)
+				{
+					DefaultColor = Console.ForegroundColor;
+					bIsWarning = Verbosity == LogEventType.Warning;
+					bIsError = Verbosity <= LogEventType.Error;
+					// @todo mono - mono doesn't seem to initialize the ForegroundColor properly, so we can't restore it properly.
+					// Avoid touching the console color unless we really need to.
+					if (bIsWarning || bIsError)
+					{
+						Console.ForegroundColor = bIsWarning ? ConsoleColor.Yellow : ConsoleColor.Red;
+					}
+				}
+				try
+				{
+					// @todo mono: mono has some kind of bug where calling mono recursively by spawning
+					// a new process causes Trace.WriteLine to stop functioning (it returns, but does nothing for some reason).
+					// work around this by simulating Trace.WriteLine on mono.
+					// We use UAT to spawn UBT instances recursively a lot, so this bug can effectively
+					// make all build output disappear outside of the top level UAT process.
+					//					#if MONO
+					lock (((System.Collections.ICollection)Trace.Listeners).SyncRoot)
+					{
+						foreach (TraceListener l in Trace.Listeners)
+						{
+							bool bIsConsole = l is Tools.DotNETCommon.ConsoleTraceListener;
+							if (Verbosity != LogEventType.Log || !bIsConsole || LogLevel >= LogEventType.Verbose)
+							{
+								List<string> Lines = FormatMessage(StackFramesToSkip + 1, Verbosity, FormatOptions, bIsConsole, Format, Args);
+								foreach (string Line in Lines)
+								{
+									l.WriteLine(Line);
+								}
+							}
+							l.Flush();
+						}
+					}
+					//					#else
+					// Call Trace directly here. Trace ensures that our logging is threadsafe using the GlobalLock.
+					//                    	Trace.WriteLine(FormatMessage(StackFramesToSkip + 1, CustomSource, Verbosity, Format, Args));
+					//					#endif
+				}
+				finally
+				{
+					// make sure we always put the console color back.
+					if (bColorConsoleOutput && (bIsWarning || bIsError))
+					{
+						Console.ForegroundColor = DefaultColor;
+					}
+				}
+			}
+		}
+
+		/// <summary>
+		/// Similar to Trace.WriteLineIf
+		/// </summary>
+		/// <param name="Condition"></param>
+		/// <param name="Verbosity"></param>
+		/// <param name="Format"></param>
+		/// <param name="Args"></param>
+		[MethodImplAttribute(MethodImplOptions.NoInlining)]
+		public static void WriteLineIf(bool Condition, LogEventType Verbosity, string Format, params object[] Args)
+		{
+			if (Condition)
+			{
+				WriteLinePrivate(1, false, Verbosity, LogFormatOptions.None, Format, Args);
+			}
+		}
+
+		/// <summary>
+		/// Mostly an internal function, but expose StackFramesToSkip to allow UAT to use existing wrapper functions and still get proper formatting.
+		/// </summary>
+		/// <param name="StackFramesToSkip"></param>
+		/// <param name="Verbosity"></param>
+		/// <param name="Format"></param>
+		/// <param name="Args"></param>
+		[MethodImplAttribute(MethodImplOptions.NoInlining)]
+		public static void WriteLine(int StackFramesToSkip, LogEventType Verbosity, string Format, params object[] Args)
+		{
+			WriteLinePrivate(StackFramesToSkip + 1, false, Verbosity, LogFormatOptions.None, Format, Args);
+		}
+		/// <summary>
+		/// Mostly an internal function, but expose StackFramesToSkip to allow UAT to use existing wrapper functions and still get proper formatting.
+		/// </summary>
+		/// <param name="StackFramesToSkip"></param>
+		/// <param name="Verbosity"></param>
+		/// <param name="FormatOptions"></param>
+		/// <param name="Format"></param>
+		/// <param name="Args"></param>
+		[MethodImplAttribute(MethodImplOptions.NoInlining)]
+		public static void WriteLine(int StackFramesToSkip, LogEventType Verbosity, LogFormatOptions FormatOptions, string Format, params object[] Args)
+		{
+			WriteLinePrivate(StackFramesToSkip + 1, false, Verbosity, FormatOptions, Format, Args);
+		}
+
+		/// <summary>
+		/// Similar to Trace.WriteLine
+		/// </summary>
+		/// <param name="Verbosity"></param>
+		/// <param name="Format"></param>
+		/// <param name="Args"></param>
+		[MethodImplAttribute(MethodImplOptions.NoInlining)]
+		public static void WriteLine(LogEventType Verbosity, string Format, params object[] Args)
+		{
+			WriteLinePrivate(1, false, Verbosity, LogFormatOptions.None, Format, Args);
+		}
+
+		/// <summary>
+		/// Similar to Trace.WriteLine
+		/// </summary>
+		/// <param name="Verbosity"></param>
+		/// <param name="FormatOptions"></param>
+		/// <param name="Format"></param>
+		/// <param name="Args"></param>
+		[MethodImplAttribute(MethodImplOptions.NoInlining)]
+		public static void WriteLine(LogEventType Verbosity, LogFormatOptions FormatOptions, string Format, params object[] Args)
+		{
+			WriteLinePrivate(1, false, Verbosity, FormatOptions, Format, Args);
+		}
+
+		/// <summary>
+		/// Writes an error message to the console.
+		/// </summary>
+		/// <param name="Format">Message format string</param>
+		/// <param name="Args">Optional arguments</param>
+		[MethodImplAttribute(MethodImplOptions.NoInlining)]
+		public static void TraceError(string Format, params object[] Args)
+		{
+			WriteLinePrivate(1, false, LogEventType.Error, LogFormatOptions.None, Format, Args);
+		}
+
+		/// <summary>
+		/// Writes a verbose message to the console.
+		/// </summary>
+		/// <param name="Format">Message format string</param>
+		/// <param name="Args">Optional arguments</param>
+		[Conditional("TRACE")]
+		[MethodImplAttribute(MethodImplOptions.NoInlining)]
+		public static void TraceVerbose(string Format, params object[] Args)
+		{
+			WriteLinePrivate(1, false, LogEventType.Verbose, LogFormatOptions.None, Format, Args);
+		}
+
+		/// <summary>
+		/// Writes a message to the console.
+		/// </summary>
+		/// <param name="Format">Message format string</param>
+		/// <param name="Args">Optional arguments</param>
+		[MethodImplAttribute(MethodImplOptions.NoInlining)]
+		public static void TraceInformation(string Format, params object[] Args)
+		{
+			WriteLinePrivate(1, false, LogEventType.Console, LogFormatOptions.None, Format, Args);
+		}
+
+		/// <summary>
+		/// Writes a warning message to the console.
+		/// </summary>
+		/// <param name="Format">Message format string</param>
+		/// <param name="Args">Optional arguments</param>
+		[MethodImplAttribute(MethodImplOptions.NoInlining)]
+		public static void TraceWarning(string Format, params object[] Args)
+		{
+			WriteLinePrivate(1, false, LogEventType.Warning, LogFormatOptions.None, Format, Args);
+		}
+
+		/// <summary>
+		/// Writes a very verbose message to the console.
+		/// </summary>
+		/// <param name="Format">Message format string</param>
+		/// <param name="Args">Optional arguments</param>
+		[Conditional("TRACE")]
+		[MethodImplAttribute(MethodImplOptions.NoInlining)]
+		public static void TraceVeryVerbose(string Format, params object[] Args)
+		{
+			WriteLinePrivate(1, false, LogEventType.VeryVerbose, LogFormatOptions.None, Format, Args);
+		}
+
+		/// <summary>
+		/// Writes a message to the log only.
+		/// </summary>
+		/// <param name="Format">Message format string</param>
+		/// <param name="Args">Optional arguments</param>
+		[Conditional("TRACE")]
+		[MethodImplAttribute(MethodImplOptions.NoInlining)]
+		public static void TraceLog(string Format, params object[] Args)
+		{
+			WriteLinePrivate(1, false, LogEventType.Log, LogFormatOptions.None, Format, Args);
+		}
+
+		/// <summary>
+		/// Similar to Trace.WriteLine
+		/// </summary>
+		/// <param name="Verbosity"></param>
+		/// <param name="Format"></param>
+		/// <param name="Args"></param>
+		[MethodImplAttribute(MethodImplOptions.NoInlining)]
+		public static void WriteLineOnce(LogEventType Verbosity, string Format, params object[] Args)
+		{
+			WriteLinePrivate(1, true, Verbosity, LogFormatOptions.None, Format, Args);
+		}
+
+		/// <summary>
+		/// Similar to Trace.WriteLine
+		/// </summary>
+		/// <param name="Verbosity"></param>
+		/// <param name="Options"></param>
+		/// <param name="Format"></param>
+		/// <param name="Args"></param>
+		[MethodImplAttribute(MethodImplOptions.NoInlining)]
+		public static void WriteLineOnce(LogEventType Verbosity, LogFormatOptions Options, string Format, params object[] Args)
+		{
+			WriteLinePrivate(1, true, Verbosity, Options, Format, Args);
+		}
+
+		/// <summary>
+		/// Writes an error message to the console.
+		/// </summary>
+		/// <param name="Format">Message format string</param>
+		/// <param name="Args">Optional arguments</param>
+		[MethodImplAttribute(MethodImplOptions.NoInlining)]
+		public static void TraceErrorOnce(string Format, params object[] Args)
+		{
+			WriteLinePrivate(1, true, LogEventType.Error, LogFormatOptions.None, Format, Args);
+		}
+
+		/// <summary>
+		/// Writes a verbose message to the console.
+		/// </summary>
+		/// <param name="Format">Message format string</param>
+		/// <param name="Args">Optional arguments</param>
+		[Conditional("TRACE")]
+		[MethodImplAttribute(MethodImplOptions.NoInlining)]
+		public static void TraceVerboseOnce(string Format, params object[] Args)
+		{
+			WriteLinePrivate(1, true, LogEventType.Verbose, LogFormatOptions.None, Format, Args);
+		}
+
+		/// <summary>
+		/// Writes a message to the console.
+		/// </summary>
+		/// <param name="Format">Message format string</param>
+		/// <param name="Args">Optional arguments</param>
+		[MethodImplAttribute(MethodImplOptions.NoInlining)]
+		public static void TraceInformationOnce(string Format, params object[] Args)
+		{
+			WriteLinePrivate(1, true, LogEventType.Console, LogFormatOptions.None, Format, Args);
+		}
+
+		/// <summary>
+		/// Writes a warning message to the console.
+		/// </summary>
+		/// <param name="Format">Message format string</param>
+		/// <param name="Args">Optional arguments</param>
+		[MethodImplAttribute(MethodImplOptions.NoInlining)]
+		public static void TraceWarningOnce(string Format, params object[] Args)
+		{
+			WriteLinePrivate(1, true, LogEventType.Warning, LogFormatOptions.None, Format, Args);
+		}
+
+		/// <summary>
+		/// Writes a very verbose message to the console.
+		/// </summary>
+		/// <param name="Format">Message format string</param>
+		/// <param name="Args">Optional arguments</param>
+		[Conditional("TRACE")]
+		[MethodImplAttribute(MethodImplOptions.NoInlining)]
+		public static void TraceVeryVerboseOnce(string Format, params object[] Args)
+		{
+			WriteLinePrivate(1, true, LogEventType.VeryVerbose, LogFormatOptions.None, Format, Args);
+		}
+
+		/// <summary>
+		/// Writes a message to the log only.
+		/// </summary>
+		/// <param name="Format">Message format string</param>
+		/// <param name="Args">Optional arguments</param>
+		[Conditional("TRACE")]
+		[MethodImplAttribute(MethodImplOptions.NoInlining)]
+		public static void TraceLogOnce(string Format, params object[] Args)
+		{
+			WriteLinePrivate(1, true, LogEventType.Log, LogFormatOptions.None, Format, Args);
+		}
+	}
+
+	/// <summary>
+	/// Class to apply a log indent for the lifetime of an object 
+	/// </summary>
+	public class ScopedLogIndent : IDisposable
+	{
+		/// <summary>
+		/// The previous indent
+		/// </summary>
+		string PrevIndent;
+
+		/// <summary>
+		/// Constructor
+		/// </summary>
+		/// <param name="Indent">Indent to append to the existing indent</param>
+		public ScopedLogIndent(string Indent)
+		{
+			PrevIndent = Log.Indent;
+			Log.Indent += Indent;
+		}
+
+		/// <summary>
+		/// Restore the log indent to normal
+		/// </summary>
+		public void Dispose()
+		{
+			if (PrevIndent != null)
+			{
+				Log.Indent = PrevIndent;
+				PrevIndent = null;
+			}
+		}
+	}
+}
diff --git a/Engine/Source/Programs/UnrealBuildTool/Platform/Android/AndroidAARHandler.cs b/Engine/Source/Programs/UnrealBuildTool/Platform/Android/AndroidAARHandler.cs
--- a/Engine/Source/Programs/UnrealBuildTool/Platform/Android/AndroidAARHandler.cs
+++ b/Engine/Source/Programs/UnrealBuildTool/Platform/Android/AndroidAARHandler.cs
@@ -1,4 +1,4 @@
-﻿// Copyright 1998-2018 Epic Games, Inc. All Rights Reserved.
+// Copyright 1998-2018 Epic Games, Inc. All Rights Reserved.
 
 using System;
 using System.Collections.Generic;
@@ -8,6 +8,7 @@ using System.Xml;
 using System.Xml.Linq;
 using System.Diagnostics;
 using System.IO;
+using Tools.DotNETCommon;
 
 using Ionic.Zip;
 
diff --git a/Engine/Source/Programs/UnrealBuildTool/Platform/HTML5/HTML5ProjectGenerator.cs b/Engine/Source/Programs/UnrealBuildTool/Platform/HTML5/HTML5ProjectGenerator.cs
--- a/Engine/Source/Programs/UnrealBuildTool/Platform/HTML5/HTML5ProjectGenerator.cs
+++ b/Engine/Source/Programs/UnrealBuildTool/Platform/HTML5/HTML5ProjectGenerator.cs
@@ -5,6 +5,7 @@ using System.Collections.Generic;
 using System.Text;
 using System.Diagnostics;
 using System.IO;
+using Tools.DotNETCommon;
 
 namespace UnrealBuildTool
 {
diff --git a/Engine/Source/Programs/UnrealBuildTool/Platform/IOS/IOSProjectGenerator.cs b/Engine/Source/Programs/UnrealBuildTool/Platform/IOS/IOSProjectGenerator.cs
--- a/Engine/Source/Programs/UnrealBuildTool/Platform/IOS/IOSProjectGenerator.cs
+++ b/Engine/Source/Programs/UnrealBuildTool/Platform/IOS/IOSProjectGenerator.cs
@@ -1,10 +1,11 @@
-﻿// Copyright 1998-2018 Epic Games, Inc. All Rights Reserved.
+// Copyright 1998-2018 Epic Games, Inc. All Rights Reserved.
 
 using System;
 using System.Collections.Generic;
 using System.Text;
 using System.Diagnostics;
 using System.IO;
+using Tools.DotNETCommon;
 
 namespace UnrealBuildTool
 {
diff --git a/Engine/Source/Programs/UnrealBuildTool/Platform/Linux/LinuxCommon.cs b/Engine/Source/Programs/UnrealBuildTool/Platform/Linux/LinuxCommon.cs
--- a/Engine/Source/Programs/UnrealBuildTool/Platform/Linux/LinuxCommon.cs
+++ b/Engine/Source/Programs/UnrealBuildTool/Platform/Linux/LinuxCommon.cs
@@ -7,6 +7,7 @@ using System.Diagnostics;
 using System.IO;
 using System.Linq;
 using Microsoft.Win32;
+using Tools.DotNETCommon;
 
 namespace UnrealBuildTool
 {
diff --git a/Engine/Source/Programs/UnrealBuildTool/Platform/Linux/LinuxProjectGenerator.cs b/Engine/Source/Programs/UnrealBuildTool/Platform/Linux/LinuxProjectGenerator.cs
--- a/Engine/Source/Programs/UnrealBuildTool/Platform/Linux/LinuxProjectGenerator.cs
+++ b/Engine/Source/Programs/UnrealBuildTool/Platform/Linux/LinuxProjectGenerator.cs
@@ -1,10 +1,11 @@
-﻿// Copyright 1998-2018 Epic Games, Inc. All Rights Reserved.
+// Copyright 1998-2018 Epic Games, Inc. All Rights Reserved.
 
 using System;
 using System.Collections.Generic;
 using System.Text;
 using System.Diagnostics;
 using System.IO;
+using Tools.DotNETCommon;
 
 namespace UnrealBuildTool
 {
diff --git a/Engine/Source/Programs/UnrealBuildTool/Platform/Mac/MacProjectGenerator.cs b/Engine/Source/Programs/UnrealBuildTool/Platform/Mac/MacProjectGenerator.cs
--- a/Engine/Source/Programs/UnrealBuildTool/Platform/Mac/MacProjectGenerator.cs
+++ b/Engine/Source/Programs/UnrealBuildTool/Platform/Mac/MacProjectGenerator.cs
@@ -5,6 +5,7 @@ using System.Collections.Generic;
 using System.Text;
 using System.Diagnostics;
 using System.IO;
+using Tools.DotNETCommon;
 
 namespace UnrealBuildTool
 {
diff --git a/Engine/Source/Programs/UnrealBuildTool/Platform/TVOS/TVOSProjectGenerator.cs b/Engine/Source/Programs/UnrealBuildTool/Platform/TVOS/TVOSProjectGenerator.cs
--- a/Engine/Source/Programs/UnrealBuildTool/Platform/TVOS/TVOSProjectGenerator.cs
+++ b/Engine/Source/Programs/UnrealBuildTool/Platform/TVOS/TVOSProjectGenerator.cs
@@ -1,10 +1,11 @@
-﻿// Copyright 1998-2018 Epic Games, Inc. All Rights Reserved.
+// Copyright 1998-2018 Epic Games, Inc. All Rights Reserved.
 
 using System;
 using System.Collections.Generic;
 using System.Text;
 using System.Diagnostics;
 using System.IO;
+using Tools.DotNETCommon;
 
 namespace UnrealBuildTool
 {
diff --git a/Engine/Source/Programs/UnrealBuildTool/Platform/Windows/WindowsProjectGenerator.cs b/Engine/Source/Programs/UnrealBuildTool/Platform/Windows/WindowsProjectGenerator.cs
--- a/Engine/Source/Programs/UnrealBuildTool/Platform/Windows/WindowsProjectGenerator.cs
+++ b/Engine/Source/Programs/UnrealBuildTool/Platform/Windows/WindowsProjectGenerator.cs
@@ -5,6 +5,7 @@ using System.Collections.Generic;
 using System.Text;
 using System.Diagnostics;
 using System.IO;
+using Tools.DotNETCommon;
 
 namespace UnrealBuildTool
 {
diff --git a/Engine/Source/Programs/UnrealBuildTool/System/BuildHostPlatform.cs b/Engine/Source/Programs/UnrealBuildTool/System/BuildHostPlatform.cs
--- a/Engine/Source/Programs/UnrealBuildTool/System/BuildHostPlatform.cs
+++ b/Engine/Source/Programs/UnrealBuildTool/System/BuildHostPlatform.cs
@@ -1,4 +1,4 @@
-﻿// Copyright 1998-2018 Epic Games, Inc. All Rights Reserved.
+// Copyright 1998-2018 Epic Games, Inc. All Rights Reserved.
 
 using System;
 using System.Collections.Generic;
@@ -7,6 +7,7 @@ using System.Text;
 using System.IO;
 using System.Diagnostics;
 using System.Runtime.InteropServices;
+using Tools.DotNETCommon;
 
 namespace UnrealBuildTool
 {
diff --git a/Engine/Source/Programs/UnrealBuildTool/System/Distcc.cs b/Engine/Source/Programs/UnrealBuildTool/System/Distcc.cs
--- a/Engine/Source/Programs/UnrealBuildTool/System/Distcc.cs
+++ b/Engine/Source/Programs/UnrealBuildTool/System/Distcc.cs
@@ -8,6 +8,7 @@ using System.Xml;
 using System.Text.RegularExpressions;
 using System.Linq;
 using System.Reflection;
+using Tools.DotNETCommon;
 
 namespace UnrealBuildTool
 {
diff --git a/Engine/Source/Programs/UnrealBuildTool/System/ExceptionUtils.cs b/Engine/Source/Programs/UnrealBuildTool/System/ExceptionUtils.cs
--- a/Engine/Source/Programs/UnrealBuildTool/System/ExceptionUtils.cs
+++ b/Engine/Source/Programs/UnrealBuildTool/System/ExceptionUtils.cs
@@ -5,6 +5,7 @@ using System.Collections.Generic;
 using System.Linq;
 using System.Text;
 using System.Threading.Tasks;
+using Tools.DotNETCommon;
 
 namespace UnrealBuildTool
 {
diff --git a/Engine/Source/Programs/UnrealBuildTool/System/JunkDeleter.cs b/Engine/Source/Programs/UnrealBuildTool/System/JunkDeleter.cs
--- a/Engine/Source/Programs/UnrealBuildTool/System/JunkDeleter.cs
+++ b/Engine/Source/Programs/UnrealBuildTool/System/JunkDeleter.cs
@@ -1,10 +1,11 @@
-﻿// Copyright 1998-2018 Epic Games, Inc. All Rights Reserved.
+// Copyright 1998-2018 Epic Games, Inc. All Rights Reserved.
 
 using System;
 using System.Collections.Generic;
 using System.Linq;
 using System.Text;
 using System.IO;
+using Tools.DotNETCommon;
 
 namespace UnrealBuildTool
 {
diff --git a/Engine/Source/Programs/UnrealBuildTool/System/LocalExecutor.cs b/Engine/Source/Programs/UnrealBuildTool/System/LocalExecutor.cs
--- a/Engine/Source/Programs/UnrealBuildTool/System/LocalExecutor.cs
+++ b/Engine/Source/Programs/UnrealBuildTool/System/LocalExecutor.cs
@@ -9,6 +9,7 @@ using System.IO;
 using System.Reflection;
 using System.Threading;
 using System.Linq;
+using Tools.DotNETCommon;
 
 namespace UnrealBuildTool
 {
diff --git a/Engine/Source/Programs/UnrealBuildTool/System/Log.cs b/Engine/Source/Programs/UnrealBuildTool/System/Log.cs
deleted file mode 100644
diff --git a/Engine/Source/Programs/UnrealBuildTool/System/ProgressWriter.cs b/Engine/Source/Programs/UnrealBuildTool/System/ProgressWriter.cs
--- a/Engine/Source/Programs/UnrealBuildTool/System/ProgressWriter.cs
+++ b/Engine/Source/Programs/UnrealBuildTool/System/ProgressWriter.cs
@@ -1,10 +1,11 @@
-﻿// Copyright 1998-2018 Epic Games, Inc. All Rights Reserved.
+// Copyright 1998-2018 Epic Games, Inc. All Rights Reserved.
 
 using System;
 using System.Collections.Generic;
 using System.Linq;
 using System.Text;
 using System.Threading.Tasks;
+using Tools.DotNETCommon;
 
 namespace UnrealBuildTool
 {
diff --git a/Engine/Source/Programs/UnrealBuildTool/System/RPCUtilHelper.cs b/Engine/Source/Programs/UnrealBuildTool/System/RPCUtilHelper.cs
--- a/Engine/Source/Programs/UnrealBuildTool/System/RPCUtilHelper.cs
+++ b/Engine/Source/Programs/UnrealBuildTool/System/RPCUtilHelper.cs
@@ -14,6 +14,7 @@ using System.Threading;
 using System.Net.Sockets;
 using Ionic.Zip;
 using DotNETUtilities;
+using Tools.DotNETCommon;
 
 namespace UnrealBuildTool
 {
diff --git a/Engine/Source/Programs/UnrealBuildTool/System/ScopedTimer.cs b/Engine/Source/Programs/UnrealBuildTool/System/ScopedTimer.cs
--- a/Engine/Source/Programs/UnrealBuildTool/System/ScopedTimer.cs
+++ b/Engine/Source/Programs/UnrealBuildTool/System/ScopedTimer.cs
@@ -1,10 +1,11 @@
-﻿// Copyright 1998-2018 Epic Games, Inc. All Rights Reserved.
+// Copyright 1998-2018 Epic Games, Inc. All Rights Reserved.
 
 using System;
 using System.Collections.Generic;
 using System.Linq;
 using System.Text;
 using System.Threading.Tasks;
+using Tools.DotNETCommon;
 
 namespace UnrealBuildTool
 {
diff --git a/Engine/Source/Programs/UnrealBuildTool/System/Telemetry.cs b/Engine/Source/Programs/UnrealBuildTool/System/Telemetry.cs
--- a/Engine/Source/Programs/UnrealBuildTool/System/Telemetry.cs
+++ b/Engine/Source/Programs/UnrealBuildTool/System/Telemetry.cs
@@ -1,4 +1,4 @@
-﻿// Copyright 1998-2018 Epic Games, Inc. All Rights Reserved.
+// Copyright 1998-2018 Epic Games, Inc. All Rights Reserved.
 
 using System;
 using System.Collections.Generic;
@@ -6,6 +6,7 @@ using System.IO;
 using System.Linq;
 using System.Reflection;
 using System.Text;
+using Tools.DotNETCommon;
 
 namespace UnrealBuildTool
 {
diff --git a/Engine/Source/Programs/UnrealBuildTool/System/UEConsoleTraceListener.cs b/Engine/Source/Programs/UnrealBuildTool/System/UEConsoleTraceListener.cs
deleted file mode 100644
diff --git a/Engine/Source/Programs/UnrealBuildTool/UnrealBuildTool.cs b/Engine/Source/Programs/UnrealBuildTool/UnrealBuildTool.cs
--- a/Engine/Source/Programs/UnrealBuildTool/UnrealBuildTool.cs
+++ b/Engine/Source/Programs/UnrealBuildTool/UnrealBuildTool.cs
@@ -370,7 +370,7 @@ namespace UnrealBuildTool
                 bColorConsoleOutput: true,
                 TraceListeners: new TraceListener[]
                 {
-                    new UEConsoleTraceListener(),
+                    new Tools.DotNETCommon.ConsoleTraceListener(),
 					StartupListener
 				}
 			);
diff --git a/Engine/Source/Programs/UnrealBuildTool/UnrealBuildTool.csproj b/Engine/Source/Programs/UnrealBuildTool/UnrealBuildTool.csproj
--- a/Engine/Source/Programs/UnrealBuildTool/UnrealBuildTool.csproj
+++ b/Engine/Source/Programs/UnrealBuildTool/UnrealBuildTool.csproj
@@ -175,7 +175,6 @@
     <Compile Include="System\PrecompiledManifest.cs" />
     <Compile Include="System\PrecompileMode.cs" />
     <Compile Include="System\RulesDocumentation.cs" />
-    <Compile Include="System\Log.cs" />
     <Compile Include="System\PrecompiledHeaderInstance.cs" />
     <Compile Include="System\PrecompiledHeaderTemplate.cs" />
     <Compile Include="Configuration\TargetInfo.cs" />
@@ -206,7 +205,6 @@
     <Compile Include="System\SupportedPlatformsAttribute.cs" />
     <Compile Include="System\TargetReceipt.cs" />
     <Compile Include="System\UBTMakefile.cs" />
-    <Compile Include="System\UEConsoleTraceListener.cs" />
     <Compile Include="System\VersionNumber.cs" />
     <Compile Include="System\XmlConfig.cs" />
     <Compile Include="System\XmlConfigData.cs" />
