Fix FileReference.Directory not returning a directory with a trailing backslash for files in the root directory.

diff --git a/Engine/Source/Programs/DotNETCommon/DotNETUtilities/DirectoryReference.cs b/Engine/Source/Programs/DotNETCommon/DotNETUtilities/DirectoryReference.cs
--- a/Engine/Source/Programs/DotNETCommon/DotNETUtilities/DirectoryReference.cs
+++ b/Engine/Source/Programs/DotNETCommon/DotNETUtilities/DirectoryReference.cs
@@ -109,6 +109,10 @@ namespace Tools.DotNETCommon
 		public static DirectoryReference GetParentDirectory(FileReference File)
 		{
 			int ParentLength = File.CanonicalName.LastIndexOf(Path.DirectorySeparatorChar);
+			if(ParentLength == 2 && File.CanonicalName[1] == ':')
+			{
+				ParentLength++;
+			}
 			return new DirectoryReference(File.FullName.Substring(0, ParentLength), File.CanonicalName.Substring(0, ParentLength));
 		}
 
