Move ExceptionUtils into DotNETUtilities.

diff --git a/Engine/Source/Programs/DotNETCommon/DotNETUtilities/DotNETUtilities.csproj b/Engine/Source/Programs/DotNETCommon/DotNETUtilities/DotNETUtilities.csproj
--- a/Engine/Source/Programs/DotNETCommon/DotNETUtilities/DotNETUtilities.csproj
+++ b/Engine/Source/Programs/DotNETCommon/DotNETUtilities/DotNETUtilities.csproj
@@ -60,6 +60,7 @@
     <Compile Include="ConsoleTraceListener.cs" />
     <Compile Include="CsProjectInfo.cs" />
     <Compile Include="DirectoryReference.cs" />
+    <Compile Include="ExceptionUtils.cs" />
     <Compile Include="FastJSON\Formatter.cs" />
     <Compile Include="FastJSON\Getters.cs" />
     <Compile Include="FastJSON\JSON.cs" />
diff --git a/Engine/Source/Programs/DotNETCommon/DotNETUtilities/ExceptionUtils.cs b/Engine/Source/Programs/DotNETCommon/DotNETUtilities/ExceptionUtils.cs
new file mode 100644
--- /dev/null
+++ b/Engine/Source/Programs/DotNETCommon/DotNETUtilities/ExceptionUtils.cs
@@ -0,0 +1,159 @@
+// Copyright 1998-2018 Epic Games, Inc. All Rights Reserved.
+
+using System;
+using System.Collections.Generic;
+using System.Linq;
+using System.Text;
+using System.Threading.Tasks;
+using Tools.DotNETCommon;
+
+namespace Tools.DotNETCommon
+{
+	/// <summary>
+	/// Methods for adding context information to exceptions
+	/// </summary>
+	public static class ExceptionUtils
+	{
+		/// <summary>
+		/// Unique key name for adding context to exceptions thrown inside Epic apps
+		/// </summary>
+		const string ContextEntryName = "EpicGames.Context";
+
+		/// <summary>
+		/// Adds a context message to a list stored on the given exception. Intended to be used in situations where supplying additional context
+		/// for an exception is valuable, but wrapping it would remove information.
+		/// </summary>
+		/// <param name="Ex">The exception that was thrown</param>
+		/// <param name="Message">Message to append to the context list</param>
+		public static void AddContext(Exception Ex, string Message)
+		{
+			List<string> Messages = Ex.Data[ContextEntryName] as List<string>;
+			if (Messages == null)
+			{
+				Messages = new List<string>();
+				Ex.Data[ContextEntryName] = Messages;
+			}
+			Messages.Add(Message);
+		}
+
+		/// <summary>
+		/// Adds a context message to a list stored on the given exception. Intended to be used in situations where supplying additional context
+		/// for an exception is valuable, but wrapping it would remove information.
+		/// </summary>
+		/// <param name="Ex">The exception that was thrown</param>
+		/// <param name="Format">Formatting string for </param>
+		/// <param name="Args">Message to append to the context list</param>
+		public static void AddContext(Exception Ex, string Format, params object[] Args)
+		{
+			AddContext(Ex, String.Format(Format, Args));
+		}
+
+		/// <summary>
+		/// Enumerates the context lines from the given exception
+		/// </summary>
+		/// <param name="Ex">The exception to retrieve context from</param>
+		/// <returns>Sequence of context lines</returns>
+		public static IEnumerable<string> GetContext(Exception Ex)
+		{
+			List<string> Messages = Ex.Data[ContextEntryName] as List<string>;
+			if (Messages != null)
+			{
+				foreach (string Message in Messages)
+				{
+					yield return Message;
+				}
+			}
+		}
+
+		/// <summary>
+		/// Prints a summary of the given exception to the console, and writes detailed information to the log
+		/// </summary>
+		/// <param name="Ex">The exception which has occurred</param>
+		/// <param name="LogFileName">Log filename to reference for more detailed information</param>
+		public static void PrintExceptionInfo(Exception Ex, string LogFileName)
+		{
+			PrintExceptionInfo(Ex, LogFileName, x => Log.TraceError("{0}", x), x => Log.TraceLog("{0}", x));
+		}
+
+		/// <summary>
+		/// Prints a summary of the given exception to the console, and writes detailed information to the log
+		/// </summary>
+		/// <param name="Ex">The exception which has occurred</param>
+		/// <param name="LogFileName">Log filename to reference for more detailed information</param>
+		/// <param name="LogAll">Output function for all verbosity</param>
+		public static void PrintExceptionInfo(Exception Ex, string LogFileName, Action<string> LogAll)
+		{
+			PrintExceptionInfo(Ex, LogFileName, LogAll, LogAll);
+		}
+
+		/// <summary>
+		/// Prints a summary of the given exception to the console, and writes detailed information to the log
+		/// </summary>
+		/// <param name="Ex">The exception which has occurred</param>
+		/// <param name="LogFileName">Log filename to reference for more detailed information</param>
+		/// <param name="LogNormal">Basic output function</param>
+		/// <param name="LogVerbose">Verbose output function</param>
+		public static void PrintExceptionInfo(Exception Ex, string LogFileName, Action<string> LogNormal, Action<string> LogVerbose)
+		{
+			LogVerbose("==============================================================================");
+
+			StringBuilder ErrorMessage = new StringBuilder();
+			if (Ex is AggregateException)
+			{
+				Exception InnerException = ((AggregateException)Ex).InnerException;
+				ErrorMessage.Append(InnerException.ToString());
+				foreach (string Line in GetContext(InnerException))
+				{
+					ErrorMessage.AppendFormat("\n  {0}", Line);
+				}
+			}
+			else
+			{
+				ErrorMessage.Append(Ex.ToString());
+			}
+			foreach (string Line in GetContext(Ex))
+			{
+				ErrorMessage.AppendFormat("\n{0}", Line);
+			}
+			if(LogFileName != null)
+			{
+				ErrorMessage.AppendFormat("\n(see {0} for full exception trace)", LogFileName);
+			}
+			LogNormal(ErrorMessage.ToString());
+
+			LogVerbose("");
+
+			List<Exception> ExceptionStack = new List<Exception>();
+			for (Exception CurrentEx = Ex; CurrentEx != null; CurrentEx = CurrentEx.InnerException)
+			{
+				ExceptionStack.Add(CurrentEx);
+			}
+			for (int Idx = ExceptionStack.Count - 1; Idx >= 0; Idx--)
+			{
+				Exception CurrentEx = ExceptionStack[Idx];
+				LogVerbose(String.Format("{0}{1}: {2}\n{3}", (Idx == ExceptionStack.Count - 1) ? "" : "Wrapped by ", CurrentEx.GetType().Name, CurrentEx.Message, CurrentEx.StackTrace));
+
+				if (CurrentEx.Data.Count > 0)
+				{
+					foreach (object Key in CurrentEx.Data.Keys)
+					{
+						object Value = CurrentEx.Data[Key];
+
+						string ValueString;
+						if(Value is List<string>)
+						{
+							ValueString = String.Format("({0})", String.Join(", ", ((List<string>)Value).Select(x => String.Format("\"{0}\"", x))));
+						}
+						else
+						{
+							ValueString = Value.ToString();
+						}
+
+						LogVerbose(String.Format("   data: {0} = {1}", Key, ValueString));
+					}
+				}
+			}
+			LogVerbose("==============================================================================");
+		}
+	}
+}
diff --git a/Engine/Source/Programs/UnrealBuildTool/System/ExceptionUtils.cs b/Engine/Source/Programs/UnrealBuildTool/System/ExceptionUtils.cs
deleted file mode 100644
diff --git a/Engine/Source/Programs/UnrealBuildTool/UnrealBuildTool.csproj b/Engine/Source/Programs/UnrealBuildTool/UnrealBuildTool.csproj
--- a/Engine/Source/Programs/UnrealBuildTool/UnrealBuildTool.csproj
+++ b/Engine/Source/Programs/UnrealBuildTool/UnrealBuildTool.csproj
@@ -169,7 +169,6 @@
     <Compile Include="System\BuildVersion.cs" />
     <Compile Include="System\CppCompileEnvironment.cs" />
     <Compile Include="System\DataDrivenPlatformInfo.cs" />
-    <Compile Include="System\ExceptionUtils.cs" />
     <Compile Include="System\ModuleManifest.cs" />
     <Compile Include="System\PluginReferenceDescriptor.cs" />
     <Compile Include="System\PrecompiledManifest.cs" />
