Add a Log.WriteException() method to dump an exception message to the console (and write the exception trace to the log)

diff --git a/Engine/Source/Programs/AutomationTool/Program.cs b/Engine/Source/Programs/AutomationTool/Program.cs
--- a/Engine/Source/Programs/AutomationTool/Program.cs
+++ b/Engine/Source/Programs/AutomationTool/Program.cs
@@ -102,13 +102,13 @@ namespace AutomationTool
 			catch (AutomationException Ex)
 			{
 				// Take the exit code from the exception
-				ExceptionUtils.PrintExceptionInfo(Ex, LogUtils.FinalLogFileName);
+				Log.WriteException(Ex, LogUtils.FinalLogFileName);
 				ReturnCode = Ex.ErrorCode;
 			}
 			catch (Exception Ex)
 			{
 				// Use a default exit code
-				ExceptionUtils.PrintExceptionInfo(Ex, LogUtils.FinalLogFileName);
+				Log.WriteException(Ex, LogUtils.FinalLogFileName);
 				ReturnCode = ExitCode.Error_Unknown;
 			}
             finally
diff --git a/Engine/Source/Programs/DotNETCommon/DotNETUtilities/ExceptionUtils.cs b/Engine/Source/Programs/DotNETCommon/DotNETUtilities/ExceptionUtils.cs
--- a/Engine/Source/Programs/DotNETCommon/DotNETUtilities/ExceptionUtils.cs
+++ b/Engine/Source/Programs/DotNETCommon/DotNETUtilities/ExceptionUtils.cs
@@ -66,37 +66,12 @@ namespace Tools.DotNETCommon
 		}
 
 		/// <summary>
-		/// Prints a summary of the given exception to the console, and writes detailed information to the log
+		/// Formats an exception for display in the log, including additional lines of context that were attached to it.
 		/// </summary>
-		/// <param name="Ex">The exception which has occurred</param>
-		/// <param name="LogFileName">Log filename to reference for more detailed information</param>
-		public static void PrintExceptionInfo(Exception Ex, string LogFileName)
+		/// <param name="Ex">The exception to format</param>
+		/// <returns>String containing the exception information. May be multiple lines.</returns>
+		public static string FormatException(Exception Ex)
 		{
-			PrintExceptionInfo(Ex, LogFileName, x => Log.TraceError("{0}", x), x => Log.TraceLog("{0}", x));
-		}
-
-		/// <summary>
-		/// Prints a summary of the given exception to the console, and writes detailed information to the log
-		/// </summary>
-		/// <param name="Ex">The exception which has occurred</param>
-		/// <param name="LogFileName">Log filename to reference for more detailed information</param>
-		/// <param name="LogAll">Output function for all verbosity</param>
-		public static void PrintExceptionInfo(Exception Ex, string LogFileName, Action<string> LogAll)
-		{
-			PrintExceptionInfo(Ex, LogFileName, LogAll, LogAll);
-		}
-
-		/// <summary>
-		/// Prints a summary of the given exception to the console, and writes detailed information to the log
-		/// </summary>
-		/// <param name="Ex">The exception which has occurred</param>
-		/// <param name="LogFileName">Log filename to reference for more detailed information</param>
-		/// <param name="LogNormal">Basic output function</param>
-		/// <param name="LogVerbose">Verbose output function</param>
-		public static void PrintExceptionInfo(Exception Ex, string LogFileName, Action<string> LogNormal, Action<string> LogVerbose)
-		{
-			LogVerbose("==============================================================================");
-
 			StringBuilder ErrorMessage = new StringBuilder();
 			if (Ex is AggregateException)
 			{
@@ -115,23 +90,27 @@ namespace Tools.DotNETCommon
 			{
 				ErrorMessage.AppendFormat("\n{0}", Line);
 			}
-			if(LogFileName != null)
-			{
-				ErrorMessage.AppendFormat("\n(see {0} for full exception trace)", LogFileName);
-			}
-			LogNormal(ErrorMessage.ToString());
-
-			LogVerbose("");
+			return ErrorMessage.ToString();
+		}
 
+		/// <summary>
+		/// Formats a detailed information about where an exception occurs, including any inner exceptions
+		/// </summary>
+		/// <param name="Ex">The exception to format</param>
+		/// <returns>String containing the exception information. May be multiple lines.</returns>
+		public static string FormatExceptionDetails(Exception Ex)
+		{
 			List<Exception> ExceptionStack = new List<Exception>();
 			for (Exception CurrentEx = Ex; CurrentEx != null; CurrentEx = CurrentEx.InnerException)
 			{
 				ExceptionStack.Add(CurrentEx);
 			}
+
+			StringBuilder Message = new StringBuilder();
 			for (int Idx = ExceptionStack.Count - 1; Idx >= 0; Idx--)
 			{
 				Exception CurrentEx = ExceptionStack[Idx];
-				LogVerbose(String.Format("{0}{1}: {2}\n{3}", (Idx == ExceptionStack.Count - 1) ? "" : "Wrapped by ", CurrentEx.GetType().Name, CurrentEx.Message, CurrentEx.StackTrace));
+				Message.AppendFormat("{0}{1}: {2}\n{3}", (Idx == ExceptionStack.Count - 1) ? "" : "Wrapped by ", CurrentEx.GetType().Name, CurrentEx.Message, CurrentEx.StackTrace);
 
 				if (CurrentEx.Data.Count > 0)
 				{
@@ -149,11 +128,11 @@ namespace Tools.DotNETCommon
 							ValueString = Value.ToString();
 						}
 
-						LogVerbose(String.Format("   data: {0} = {1}", Key, ValueString));
+						Message.AppendFormat("   data: {0} = {1}", Key, ValueString);
 					}
 				}
 			}
-			LogVerbose("==============================================================================");
+			return Message.ToString();
 		}
 	}
 }
diff --git a/Engine/Source/Programs/DotNETCommon/DotNETUtilities/Log.cs b/Engine/Source/Programs/DotNETCommon/DotNETUtilities/Log.cs
--- a/Engine/Source/Programs/DotNETCommon/DotNETUtilities/Log.cs
+++ b/Engine/Source/Programs/DotNETCommon/DotNETUtilities/Log.cs
@@ -472,6 +472,20 @@ namespace Tools.DotNETCommon
 			WriteLinePrivate(1, false, Verbosity, FormatOptions, Format, Args);
 		}
 
+		/// <summary>
+		/// Formats an exception for display in the log. The exception message is shown as an error, and the stack trace is included in the log.
+		/// </summary>
+		/// <param name="Ex">The exception to display</param>
+		/// <param name="LogFileName">The log filename to display, if any</param>
+		public static void WriteException(Exception Ex, string LogFileName)
+		{
+			string LogSuffix = (LogFileName == null)? "" : String.Format("\n(see {0} for full exception trace)", LogFileName);
+			TraceLog("==============================================================================");
+			TraceError("{0}{1}", ExceptionUtils.FormatException(Ex), LogSuffix);
+			TraceLog("\n{0}", ExceptionUtils.FormatExceptionDetails(Ex));
+			TraceLog("==============================================================================");
+		}
+
 		/// <summary>
 		/// Writes an error message to the console.
 		/// </summary>
diff --git a/Engine/Source/Programs/UnrealBuildTool/System/LocalExecutor.cs b/Engine/Source/Programs/UnrealBuildTool/System/LocalExecutor.cs
--- a/Engine/Source/Programs/UnrealBuildTool/System/LocalExecutor.cs
+++ b/Engine/Source/Programs/UnrealBuildTool/System/LocalExecutor.cs
@@ -193,7 +193,7 @@ namespace UnrealBuildTool
 					catch (Exception ex)
 					{
 						Log.TraceError("Failed to start local process for action: {0} {1}", Action.CommandPath, Action.CommandArguments);
-						ExceptionUtils.PrintExceptionInfo(ex, null);
+						Log.WriteException(ex, null);
 						ExitCode = 1;
 						bComplete = true;
 						return;
diff --git a/Engine/Source/Programs/UnrealBuildTool/UnrealBuildTool.cs b/Engine/Source/Programs/UnrealBuildTool/UnrealBuildTool.cs
--- a/Engine/Source/Programs/UnrealBuildTool/UnrealBuildTool.cs
+++ b/Engine/Source/Programs/UnrealBuildTool/UnrealBuildTool.cs
@@ -1725,7 +1725,7 @@ namespace UnrealBuildTool
 				{
 					throw;
 				}
-				ExceptionUtils.PrintExceptionInfo(Ex, String.IsNullOrEmpty(BuildConfiguration.LogFileName) ? null : BuildConfiguration.LogFileName);
+				Log.WriteException(Ex, String.IsNullOrEmpty(BuildConfiguration.LogFileName) ? null : BuildConfiguration.LogFileName);
 				BuildResult = ECompilationResult.OtherCompilationError;
 			}
 
