Attribute-driven Perforce wrapper (old Epic Friday project). Offers a more complete implementation than the current P4 wrapper in UAT without requiring any platform-specific libraries. Uses the Python binary output for parsing.

diff --git a/Engine/Source/Programs/DotNETCommon/DotNETUtilities/Perforce/PerforceAttributes.cs b/Engine/Source/Programs/DotNETCommon/DotNETUtilities/Perforce/PerforceAttributes.cs
new file mode 100644
--- /dev/null
+++ b/Engine/Source/Programs/DotNETCommon/DotNETUtilities/Perforce/PerforceAttributes.cs
@@ -0,0 +1,63 @@
+using System;
+using System.Collections.Generic;
+using System.Linq;
+using System.Text;
+using System.Threading.Tasks;
+
+namespace Tools.DotNETCommon.Perforce
+{
+	/// <summary>
+	/// Attributes for fields that should be deserialized from P4 tags
+	/// </summary>
+	[AttributeUsage(AttributeTargets.Field)]
+	public class PerforceTagAttribute : Attribute
+	{
+		/// <summary>
+		/// The tag name
+		/// </summary>
+		public string Name;
+
+		/// <summary>
+		/// Whether this tag is required for a valid record
+		/// </summary>
+		public bool Optional = false;
+
+		/// <summary>
+		/// Constructor
+		/// </summary>
+		/// <param name="Name">Name of the tag</param>
+		public PerforceTagAttribute(string Name)
+		{
+			this.Name = Name;
+		}
+	}
+
+	/// <summary>
+	/// Specifies the name of an enum when converted into a P4 string
+	/// </summary>
+	[AttributeUsage(AttributeTargets.Field)]
+	public class PerforceEnumAttribute : Attribute
+	{
+		/// <summary>
+		/// Name of the enum value
+		/// </summary>
+		public string Name;
+
+		/// <summary>
+		/// Constructor
+		/// </summary>
+		/// <param name="Name">Name of the serialized value</param>
+		public PerforceEnumAttribute(string Name)
+		{
+			this.Name = Name;
+		}
+	}
+
+	/// <summary>
+	/// When attached to a list field, indicates that a list of structures can be included in the record
+	/// </summary>
+	[AttributeUsage(AttributeTargets.Field)]
+	public class PerforceRecordListAttribute : Attribute
+	{
+	}
+}
diff --git a/Engine/Source/Programs/DotNETCommon/DotNETUtilities/Perforce/PerforceChildProcess.cs b/Engine/Source/Programs/DotNETCommon/DotNETUtilities/Perforce/PerforceChildProcess.cs
new file mode 100644
--- /dev/null
+++ b/Engine/Source/Programs/DotNETCommon/DotNETUtilities/Perforce/PerforceChildProcess.cs
@@ -0,0 +1,231 @@
+using System;
+using System.Collections.Generic;
+using System.Diagnostics;
+using System.IO;
+using System.Linq;
+using System.Text;
+using System.Threading.Tasks;
+
+namespace Tools.DotNETCommon.Perforce
+{
+	enum PerforceRecordType
+	{
+		Invalid,
+		Stat,
+		Text,
+		Binary,
+	}
+
+	enum PerforceFieldType
+	{
+		Record,
+		EndRecord,
+		String,
+		Integer,
+		EndStream
+	}
+
+	class PerforceChildProcess : IDisposable
+	{
+		Process ChildProcess;
+		byte[] Buffer;
+		int BufferPos;
+		int BufferEnd;
+		bool bEndOfStream;
+
+		public PerforceChildProcess(byte[] InputData, string Format, params object[] Args)
+		{
+			ChildProcess = new Process();
+			ChildProcess.StartInfo.FileName = "p4.exe";
+			ChildProcess.StartInfo.Arguments = "-G " + String.Format(Format, Args);
+
+			ChildProcess.StartInfo.RedirectStandardError = true;
+			ChildProcess.StartInfo.RedirectStandardOutput = true;
+			ChildProcess.StartInfo.RedirectStandardInput = InputData != null;
+			ChildProcess.StartInfo.UseShellExecute = false;
+			ChildProcess.StartInfo.CreateNoWindow = true;
+
+			ChildProcess.Start();
+
+			if(InputData != null)
+			{
+				ChildProcess.StandardInput.BaseStream.Write(InputData, 0, InputData.Length);
+				ChildProcess.StandardInput.BaseStream.Close();
+			}
+
+			Buffer = new byte[2048];
+		}
+
+		public void Dispose()
+		{
+			if(ChildProcess != null)
+			{
+				try
+				{
+					if(!ChildProcess.HasExited)
+					{
+						ChildProcess.Kill();
+						ChildProcess.WaitForExit();
+					}
+				}
+				catch
+				{
+				}
+
+				ChildProcess.Dispose();
+				ChildProcess = null;
+			}
+		}
+
+		public bool IsEndOfStream()
+		{
+			if(!bEndOfStream)
+			{
+				bEndOfStream = (BufferPos == BufferEnd && !TryUpdateBuffer(1));
+			}
+			return bEndOfStream;
+		}
+
+		public bool TryReadRecord(List<KeyValuePair<string, object>> Record)
+		{
+			// Check if we've reached the end of the stream. This is the only condition where we return false.
+			if(IsEndOfStream())
+			{
+				return false;
+			}
+
+			// Check that a dictionary follows
+			byte Temp = ReadByte();
+			if(Temp != '{')
+			{
+				throw new PerforceException("Unexpected data while parsing marshaled output - expected '{'");
+			}
+
+			// Read all the fields in the record
+			Record.Clear();
+			for(;;)
+			{
+				// Read the next field type. Perforce only outputs string records. A '0' character indicates the end of the dictionary.
+				byte KeyFieldType = ReadByte();
+				if(KeyFieldType == '0')
+				{
+					break;
+				}
+				else if(KeyFieldType != 's')
+				{
+					throw new PerforceException("Unexpected key field type while parsing marshalled output ({0}) - expected 's'", (int)KeyFieldType);
+				}
+
+				// Read the key
+				string Key = ReadString();
+
+				// Read the value type.
+				byte ValueFieldType = ReadByte();
+				if(ValueFieldType == 'i')
+				{
+					// An integer
+					Record.Add(new KeyValuePair<string, object>(Key, ReadInt32()));
+				}
+				else if(ValueFieldType == 's')
+				{
+					// A string (or variable length data)
+					byte[] Value = ReadBytes();
+					if(Record.Count > 0 && Record[0].Key == "code" && ((string)Record[0].Value == "binary" || (string)Record[0].Value == "text"))
+					{
+						Record.Add(new KeyValuePair<string, object>(Key, Value));
+					}
+					else
+					{
+						Record.Add(new KeyValuePair<string, object>(Key, Encoding.UTF8.GetString(Value)));
+					}
+				}
+				else
+				{
+					throw new PerforceException("Unexpected value field type while parsing marshalled output ({0}) - expected 's'", (int)ValueFieldType);
+				}
+			}
+			return true;
+		}
+
+		public byte ReadByte()
+		{
+			UpdateBuffer(1);
+
+			byte NextByte = Buffer[BufferPos++];
+			return NextByte;
+		}
+
+		public int ReadInt32()
+		{
+			UpdateBuffer(4);
+
+			int Value = Buffer[BufferPos + 0] | (Buffer[BufferPos + 1] << 8) | (Buffer[BufferPos + 2] << 16) | (Buffer[BufferPos + 3] << 24);
+			BufferPos += 4;
+			return Value;
+		}
+
+		public string ReadString()
+		{
+			byte[] Bytes = ReadBytes();
+			return Encoding.UTF8.GetString(Bytes);
+		}
+
+		public byte[] ReadBytes()
+		{
+			int MaxNumBytes = ReadInt32();
+
+			byte[] Buffer = new byte[MaxNumBytes];
+			CopyBytes(MaxNumBytes, new MemoryStream(Buffer, true));
+			return Buffer;
+		}
+
+		public void CopyBytes(int NumBytes, Stream OutputStream)
+		{
+			while(NumBytes > 0)
+			{
+				UpdateBuffer(1);
+				int NumBytesThisLoop = Math.Min(NumBytes, BufferEnd - BufferPos);
+				OutputStream.Write(Buffer, BufferPos, NumBytesThisLoop);
+				BufferPos += NumBytesThisLoop;
+				NumBytes -= NumBytesThisLoop;
+			}
+		}
+
+		private void UpdateBuffer(int MinSize)
+		{
+			if(!TryUpdateBuffer(MinSize))
+			{
+				throw new PerforceException("Unexpected end of stream");
+			}
+		}
+
+		private bool TryUpdateBuffer(int MinSize)
+		{
+			while(BufferPos + MinSize > BufferEnd)
+			{
+				if(BufferPos > 0)
+				{
+					Array.Copy(Buffer, BufferPos, Buffer, 0, BufferEnd - BufferPos);
+					BufferEnd -= BufferPos;
+					BufferPos = 0;
+				}
+
+				int ReadSize = ChildProcess.StandardOutput.BaseStream.Read(Buffer, BufferEnd, Buffer.Length - BufferEnd);
+				if(ReadSize == 0)
+				{
+					if(BufferPos == BufferEnd)
+					{
+						return false;
+					}
+					else
+					{
+						throw new PerforceException("Partially read record");
+					}
+				}
+
+				BufferEnd += ReadSize;
+			}
+			return true;
+		}
+	}
+}
diff --git a/Engine/Source/Programs/DotNETCommon/DotNETUtilities/Perforce/PerforceConnection.cs b/Engine/Source/Programs/DotNETCommon/DotNETUtilities/Perforce/PerforceConnection.cs
new file mode 100644
--- /dev/null
+++ b/Engine/Source/Programs/DotNETCommon/DotNETUtilities/Perforce/PerforceConnection.cs
@@ -0,0 +1,2129 @@
+using System;
+using System.Collections.Concurrent;
+using System.Collections.Generic;
+using System.Diagnostics;
+using System.IO;
+using System.Linq;
+using System.Reflection;
+using System.Text;
+using System.Threading.Tasks;
+
+namespace Tools.DotNETCommon.Perforce
+{
+	/// <summary>
+	/// Stores settings for communicating with a Perforce server.
+	/// </summary>
+	public class PerforceConnection
+	{
+		#region Plumbing
+
+		/// <summary>
+		/// Stores cached information about a field with a P4Tag attribute
+		/// </summary>
+		class CachedTagInfo
+		{
+			/// <summary>
+			/// Name of the tag. Specified in the attribute or inferred from the field name.
+			/// </summary>
+			public string Name;
+
+			/// <summary>
+			/// Whether this tag is optional or not.
+			/// </summary>
+			public bool Optional;
+
+			/// <summary>
+			/// The field containing the value of this data.
+			/// </summary>
+			public FieldInfo Field;
+
+			/// <summary>
+			/// Index into the bitmask of required types
+			/// </summary>
+			public ulong RequiredTagBitMask;
+		}
+
+		/// <summary>
+		/// Stores cached information about a record
+		/// </summary>
+		class CachedRecordInfo
+		{
+			/// <summary>
+			/// Type of the record
+			/// </summary>
+			public Type Type;
+
+			/// <summary>
+			/// Map of tag names to their cached reflection information
+			/// </summary>
+			public Dictionary<string, CachedTagInfo> TagNameToInfo = new Dictionary<string, CachedTagInfo>();
+
+			/// <summary>
+			/// Bitmask of all the required tags. Formed by bitwise-or'ing the RequiredTagBitMask fields for each required CachedTagInfo.
+			/// </summary>
+			public ulong RequiredTagsBitMask;
+
+			/// <summary>
+			/// The type of records to create for subelements
+			/// </summary>
+			public Type SubElementType;
+
+			/// <summary>
+			/// The cached record info for the subelement type
+			/// </summary>
+			public CachedRecordInfo SubElementRecordInfo;
+
+			/// <summary>
+			/// Field containing subelements
+			/// </summary>
+			public FieldInfo SubElementField;
+		}
+
+		/// <summary>
+		/// Unix epoch; used for converting times back into C# datetime objects
+		/// </summary>
+		static readonly DateTime UnixEpoch = new DateTime(1970, 1, 1, 0, 0, 0, DateTimeKind.Utc);
+
+		/// <summary>
+		/// Cached map of enum types to a lookup mapping from p4 strings to enum values.
+		/// </summary>
+		static ConcurrentDictionary<Type, Dictionary<string, int>> EnumTypeToFlags = new ConcurrentDictionary<Type, Dictionary<string, int>>();
+
+		/// <summary>
+		/// Cached set of record 
+		/// </summary>
+		static ConcurrentDictionary<Type, CachedRecordInfo> RecordTypeToInfo = new ConcurrentDictionary<Type, CachedRecordInfo>();
+
+		/// <summary>
+		/// Default type for info
+		/// </summary>
+		static CachedRecordInfo InfoRecordInfo = GetCachedRecordInfo(typeof(PerforceInfo));
+
+		/// <summary>
+		/// Default type for errors
+		/// </summary>
+		static CachedRecordInfo ErrorRecordInfo = GetCachedRecordInfo(typeof(PerforceError));
+
+		/// <summary>
+		/// Global options for each command
+		/// </summary>
+		public readonly string GlobalOptions;
+
+		/// <summary>
+		/// Constructor
+		/// </summary>
+		/// <param name="GlobalOptions">Global options to pass to every Perforce command</param>
+		public PerforceConnection(string GlobalOptions)
+		{
+			this.GlobalOptions = GlobalOptions;
+		}
+
+		/// <summary>
+		/// Execute a Perforce command and parse the output as marshalled Python objects. This is more robustly defined than the text-based tagged output
+		/// format, because it avoids ambiguity when returned fields can have newlines.
+		/// </summary>
+		/// <param name="CommandLine">Command line to execute Perforce with</param>
+		/// <param name="InputData">Input data to pass to the Perforce server. May be null.</param>
+		/// <param name="RecordHandler">Handler for each received record.</param>
+		public void Command(string CommandLine, byte[] InputData, Action<List<KeyValuePair<string, object>>> RecordHandler)
+		{
+			using(PerforceChildProcess Process = new PerforceChildProcess(InputData, "-G {0} {1}", GlobalOptions, CommandLine))
+			{
+				List<KeyValuePair<string, object>> Record = new List<KeyValuePair<string, object>>();
+				while(Process.TryReadRecord(Record))
+				{
+					RecordHandler(Record);
+				}
+			}
+		}
+
+		/// <summary>
+		/// Serializes a list of key/value pairs into binary format.
+		/// </summary>
+		/// <param name="KeyValuePairs">List of key value pairs</param>
+		/// <returns>Serialized record data</returns>
+		static byte[] SerializeRecord(List<KeyValuePair<string, object>> KeyValuePairs)
+		{
+			MemoryStream Stream = new MemoryStream();
+			using (BinaryWriter Writer = new BinaryWriter(Stream))
+			{
+				Writer.Write((byte)'{');
+				foreach(KeyValuePair<string, object> KeyValuePair in KeyValuePairs)
+				{
+					Writer.Write('s');
+					byte[] KeyBytes = Encoding.UTF8.GetBytes(KeyValuePair.Key);
+					Writer.Write((int)KeyBytes.Length);
+					Writer.Write(KeyBytes);
+
+					if (KeyValuePair.Value is string)
+					{
+						Writer.Write('s');
+						byte[] ValueBytes = Encoding.UTF8.GetBytes((string)KeyValuePair.Value);
+						Writer.Write((int)ValueBytes.Length);
+						Writer.Write(ValueBytes);
+					}
+					else
+					{
+						throw new PerforceException("Unsupported formatting type for {0}", KeyValuePair.Key);
+					}
+				}
+				Writer.Write((byte)'0');
+			}
+			return Stream.ToArray();
+		}
+
+		/// <summary>
+		/// Execute a command and parse the response
+		/// </summary>
+		/// <param name="Arguments">Arguments for the command</param>
+		/// <param name="InputData">Input data to pass to Perforce</param>
+		/// <param name="StatRecordType">The type of records to return for "stat" responses</param>
+		/// <returns>List of objects returned by the server</returns>
+		public List<PerforceResponse> Command(string Arguments, byte[] InputData, Type StatRecordType)
+		{
+			CachedRecordInfo StatRecordInfo = (StatRecordType == null)? null : GetCachedRecordInfo(StatRecordType);
+
+			List<PerforceResponse> Responses = new List<PerforceResponse>();
+			Action<List<KeyValuePair<string, object>>> Handler = (KeyValuePairs) =>
+			{
+				if(KeyValuePairs.Count == 0)
+				{
+					throw new PerforceException("Unexpected empty record returned by Perforce.");
+				}
+				if(KeyValuePairs[0].Key != "code")
+				{
+					throw new PerforceException("Expected first returned field to be 'code'");
+				}
+
+				string Code = KeyValuePairs[0].Value as string;
+
+				int Idx = 1;
+				if (Code == "stat" && StatRecordType != null)
+				{
+					Responses.Add(ParseResponse(KeyValuePairs, ref Idx, "", StatRecordInfo));
+				}
+				else if (Code == "info")
+				{
+					Responses.Add(ParseResponse(KeyValuePairs, ref Idx, "", InfoRecordInfo));
+				}
+				else if(Code == "error")
+				{
+					Responses.Add(ParseResponse(KeyValuePairs, ref Idx, "", ErrorRecordInfo));
+				}
+				else
+				{
+					throw new PerforceException("Unknown return code for record: {0}", KeyValuePairs[0].Value);
+				}
+			};
+			Command(Arguments, InputData, Handler);
+
+			return Responses;
+		}
+
+		/// <summary>
+		/// Execute a command and parse the response
+		/// </summary>
+		/// <param name="Arguments">Arguments for the command</param>
+		/// <param name="InputData">Input data to pass to Perforce</param>
+		/// <returns>List of objects returned by the server</returns>
+		public PerforceResponseList<T> Command<T>(string Arguments, byte[] InputData) where T : class
+		{
+			List<PerforceResponse> Responses = Command(Arguments, InputData, typeof(T));
+
+			PerforceResponseList<T> TypedResponses = new PerforceResponseList<T>();
+			foreach (PerforceResponse Response in Responses)
+			{
+				TypedResponses.Add(new PerforceResponse<T>(Response));
+			}
+			return TypedResponses;
+		}
+
+		/// <summary>
+		/// Attempts to execute the given command, returning the results from the server or the first PerforceResponse object.
+		/// </summary>
+		/// <param name="Arguments">Arguments for the command.</param>
+		/// <param name="InputData">Input data for the command.</param>
+		/// <param name="StatRecordType">Type of element to return in the response</param>
+		/// <returns>Response from the server; either an object of type T or error.</returns>
+		public PerforceResponse SingleResponseCommand(string Arguments, byte[] InputData, Type StatRecordType)
+		{
+			List<PerforceResponse> Responses = Command(Arguments, InputData, StatRecordType);
+			if(Responses.Count != 1)
+			{
+				throw new PerforceException("Expected one result from 'p4 {0}', got {1}", Arguments, Responses.Count);
+			}
+			return Responses[0];
+		}
+
+		/// <summary>
+		/// Attempts to execute the given command, returning the results from the server or the first PerforceResponse object.
+		/// </summary>
+		/// <typeparam name="T">Type of record to parse</typeparam>
+		/// <param name="Arguments">Arguments for the command.</param>
+		/// <param name="InputData">Input data for the command.</param>
+		/// <returns>Response from the server; either an object of type T or error.</returns>
+		public PerforceResponse<T> SingleResponseCommand<T>(string Arguments, byte[] InputData) where T : class
+		{
+			return new PerforceResponse<T>(SingleResponseCommand(Arguments, InputData, typeof(T)));
+		}
+
+		/// <summary>
+		/// Parse an individual record from the server.
+		/// </summary>
+		/// <param name="KeyValuePairs">List of tagged values returned by the server.</param>
+		/// <param name="Idx">Index of the first tagged value to parse.</param>
+		/// <param name="RequiredSuffix">The required suffix for any subobject arrays.</param>
+		/// <param name="RecordInfo">Reflection information for the type being serialized into.</param>
+		/// <returns>The parsed object.</returns>
+		PerforceResponse ParseResponse(List<KeyValuePair<string, object>> KeyValuePairs, ref int Idx, string RequiredSuffix, CachedRecordInfo RecordInfo)
+		{
+			// Create a bitmask for all the required tags
+			ulong RequiredTagsBitMask = 0;
+
+			// Get the record info, and parse it into the object
+			object NewRecord = Activator.CreateInstance(RecordInfo.Type);
+			while(Idx < KeyValuePairs.Count)
+			{
+				// Split out the tag and value
+				string Tag = KeyValuePairs[Idx].Key;
+				string Value = KeyValuePairs[Idx].Value.ToString();
+
+				// Parse the suffix from the current key
+				int SuffixIdx = Tag.Length;
+				while(SuffixIdx > 0 && (Tag[SuffixIdx - 1] == ',' || (Tag[SuffixIdx - 1] >= '0' && Tag[SuffixIdx - 1] <= '9')))
+				{
+					SuffixIdx--;
+				}
+
+				// Split out the suffix
+				string Suffix = Tag.Substring(SuffixIdx);
+				Tag = Tag.Substring(0, SuffixIdx);
+
+				// Check whether it's a subobject or part of the current object.
+				if (Suffix == RequiredSuffix)
+				{
+					// Part of the current object
+					CachedTagInfo TagInfo;
+					if (RecordInfo.TagNameToInfo.TryGetValue(Tag, out TagInfo))
+					{
+						FieldInfo FieldInfo = TagInfo.Field;
+						if (FieldInfo.FieldType == typeof(DateTime))
+						{
+							DateTime Time;
+							if(!DateTime.TryParse(Value, out Time))
+							{
+								Time = UnixEpoch + TimeSpan.FromSeconds(long.Parse(Value));
+							}
+							FieldInfo.SetValue(NewRecord, Time);
+						}
+						else if (FieldInfo.FieldType == typeof(bool))
+						{
+							FieldInfo.SetValue(NewRecord, Value.Length == 0 || Value == "true");
+						}
+						else if(FieldInfo.FieldType == typeof(Nullable<bool>))
+						{
+							FieldInfo.SetValue(NewRecord, Value == "true");
+						}
+						else if (FieldInfo.FieldType == typeof(int))
+						{
+							if(Value == "new" || Value == "none")
+							{
+								FieldInfo.SetValue(NewRecord, -1);
+							}
+							else if(Value.StartsWith("#"))
+							{
+								FieldInfo.SetValue(NewRecord, (Value == "#none") ? 0 : int.Parse(Value.Substring(1)));
+							}
+							else
+							{
+								FieldInfo.SetValue(NewRecord, int.Parse(Value));
+							}
+						}
+						else if (FieldInfo.FieldType == typeof(long))
+						{
+							FieldInfo.SetValue(NewRecord, long.Parse(Value));
+						}
+						else if (FieldInfo.FieldType == typeof(string))
+						{
+							FieldInfo.SetValue(NewRecord, Value);
+						}
+						else if(FieldInfo.FieldType.IsEnum)
+						{
+							FieldInfo.SetValue(NewRecord, ParseEnum(FieldInfo.FieldType, Value));
+						}
+						else
+						{
+							throw new PerforceException("Unsupported type of {0}.{1} for tag '{0}'", RecordInfo.Type.Name, FieldInfo.FieldType.Name, Tag);
+						}
+						RequiredTagsBitMask |= TagInfo.RequiredTagBitMask;
+					}
+					Idx++;
+				}
+				else if (Suffix.StartsWith(RequiredSuffix) && (RequiredSuffix.Length == 0 || Suffix[RequiredSuffix.Length] == ','))
+				{
+					// Part of a subobject. If this record doesn't have any listed subobject type, skip the field and continue.
+					if (RecordInfo.SubElementField == null)
+					{
+						CachedTagInfo TagInfo;
+						if (RecordInfo.TagNameToInfo.TryGetValue(Tag, out TagInfo))
+						{
+							FieldInfo FieldInfo = TagInfo.Field;
+							if (FieldInfo.FieldType == typeof(List<string>))
+							{
+								((List<string>)FieldInfo.GetValue(NewRecord)).Add(Value);
+							}
+							else
+							{
+								throw new PerforceException("Unsupported type of {0}.{1} for tag '{0}'", RecordInfo.Type.Name, FieldInfo.FieldType.Name, Tag);
+							}
+							RequiredTagsBitMask |= TagInfo.RequiredTagBitMask;
+						}
+						Idx++;
+					}
+					else
+					{
+						// Get the expected suffix for the next item based on the number of elements already in the list
+						System.Collections.IList List = (System.Collections.IList)RecordInfo.SubElementField.GetValue(NewRecord);
+						string RequiredChildSuffix = (RequiredSuffix.Length == 0) ? String.Format("{0}", List.Count) : String.Format("{0},{1}", RequiredSuffix, List.Count);
+						if (Suffix != RequiredChildSuffix)
+						{
+							throw new PerforceException("Subobject element received out of order; expected {0}{1}, got {0}{2}", Tag, RequiredChildSuffix, Suffix);
+						}
+
+						// Parse the subobject and add it to the list
+						object SubObject = ParseResponse(KeyValuePairs, ref Idx, RequiredChildSuffix, RecordInfo.SubElementRecordInfo);
+						List.Add(SubObject);
+					}
+				}
+				else
+				{
+					break;
+				}
+			}
+
+			// Make sure we've got all the required tags we need
+			if (RequiredTagsBitMask != RecordInfo.RequiredTagsBitMask)
+			{
+				string MissingTagNames = String.Join(", ", RecordInfo.TagNameToInfo.Where(x => (RequiredTagsBitMask | x.Value.RequiredTagBitMask) != RequiredTagsBitMask).Select(x => x.Key));
+				throw new PerforceException("Missing '{0}' tag when parsing '{1}'", MissingTagNames, RecordInfo.Type.Name);
+			}
+
+			return new PerforceResponse(NewRecord);
+		}
+
+		/// <summary>
+		/// Gets a mapping of flags to enum values for the given type
+		/// </summary>
+		/// <param name="EnumType">The enum type to retrieve flags for</param>
+		/// <returns>Map of name to enum value</returns>
+		static Dictionary<string, int> GetCachedEnumFlags(Type EnumType)
+		{
+			Dictionary<string, int> NameToValue;
+			if (!EnumTypeToFlags.TryGetValue(EnumType, out NameToValue))
+			{
+				NameToValue = new Dictionary<string, int>();
+
+				FieldInfo[] Fields = EnumType.GetFields(BindingFlags.Public | BindingFlags.Static);
+				foreach (FieldInfo Field in Fields)
+				{
+					PerforceEnumAttribute Attribute = Field.GetCustomAttribute<PerforceEnumAttribute>();
+					if (Attribute != null)
+					{
+						NameToValue.Add(Attribute.Name, (int)Field.GetValue(null));
+					}
+				}
+
+				if (!EnumTypeToFlags.TryAdd(EnumType, NameToValue))
+				{
+					NameToValue = EnumTypeToFlags[EnumType];
+				}
+			}
+			return NameToValue;
+		}
+
+		/// <summary>
+		/// Parses an enum value, using PerforceEnumAttribute markup for names.
+		/// </summary>
+		/// <param name="EnumType">Type of the enum to parse.</param>
+		/// <param name="Value">Value of the enum.</param>
+		/// <returns>Text for the enum.</returns>
+		string GetEnumText(Type EnumType, object Value)
+		{
+			int IntegerValue = (int)Value;
+
+			Dictionary<string, int> NameToValue = GetCachedEnumFlags(EnumType);
+			if(EnumType.GetCustomAttribute<FlagsAttribute>() != null)
+			{
+				List<string> Names = new List<string>();
+				foreach (KeyValuePair<string, int> Pair in NameToValue)
+				{
+					if ((IntegerValue & Pair.Value) != 0)
+					{
+						Names.Add(Pair.Key);
+					}
+				}
+				return String.Join(" ", Names);
+			}
+			else
+			{
+				string Name = null;
+				foreach (KeyValuePair<string, int> Pair in NameToValue)
+				{
+					if (IntegerValue == Pair.Value)
+					{
+						Name = Pair.Key;
+						break;
+					}
+				}
+				return Name;
+			}
+		}
+
+		/// <summary>
+		/// Parses an enum value, using PerforceEnumAttribute markup for names.
+		/// </summary>
+		/// <param name="EnumType">Type of the enum to parse.</param>
+		/// <param name="Text">Text to parse.</param>
+		/// <returns>The parsed enum value. Unknown values will be ignored.</returns>
+		object ParseEnum(Type EnumType, string Text)
+		{
+			Dictionary<string, int> NameToValue = GetCachedEnumFlags(EnumType);
+			if(EnumType.GetCustomAttribute<FlagsAttribute>() != null)
+			{
+				int Result = 0;
+				foreach (string Item in Text.Split(' '))
+				{
+					int ItemValue;
+					if (NameToValue.TryGetValue(Item, out ItemValue))
+					{
+						Result |= ItemValue;
+					}
+				}
+				return Enum.ToObject(EnumType, Result);
+			}
+			else
+			{
+				int Result;
+				NameToValue.TryGetValue(Text, out Result);
+				return Enum.ToObject(EnumType, Result);
+			}
+		}
+
+		/// <summary>
+		/// Gets reflection data for the given record type
+		/// </summary>
+		/// <param name="RecordType">The type to retrieve record info for</param>
+		/// <returns>The cached reflection information for the given type</returns>
+		static CachedRecordInfo GetCachedRecordInfo(Type RecordType)
+		{
+			CachedRecordInfo Record;
+			if (!RecordTypeToInfo.TryGetValue(RecordType, out Record))
+			{
+				Record = new CachedRecordInfo();
+				Record.Type = RecordType;
+
+				// Get all the fields for this type
+				FieldInfo[] Fields = RecordType.GetFields(BindingFlags.Public | BindingFlags.NonPublic | BindingFlags.Instance);
+
+				// Build the map of all tags for this record
+				foreach (FieldInfo Field in Fields)
+				{
+					PerforceTagAttribute TagAttribute = Field.GetCustomAttribute<PerforceTagAttribute>();
+					if (TagAttribute != null)
+					{
+						CachedTagInfo Tag = new CachedTagInfo();
+						Tag.Name = TagAttribute.Name ?? Field.Name;
+						Tag.Optional = TagAttribute.Optional;
+						Tag.Field = Field;
+						if(!Tag.Optional)
+						{
+							Tag.RequiredTagBitMask = Record.RequiredTagsBitMask + 1;
+							if(Tag.RequiredTagBitMask == 0)
+							{
+								throw new PerforceException("Too many required tags in {0}; max is {1}", RecordType.Name, sizeof(ulong) * 8);
+							}
+							Record.RequiredTagsBitMask |= Tag.RequiredTagBitMask;
+						}
+						Record.TagNameToInfo.Add(Tag.Name, Tag);
+					}
+
+					PerforceRecordListAttribute SubElementAttribute = Field.GetCustomAttribute<PerforceRecordListAttribute>();
+					if(SubElementAttribute != null)
+					{
+						Record.SubElementField = Field;
+						Record.SubElementType = Field.FieldType.GenericTypeArguments[0];
+						Record.SubElementRecordInfo = GetCachedRecordInfo(Record.SubElementType);
+					}
+				}
+
+				// Try to save the record info, or get the version that's already in the cache
+				if(!RecordTypeToInfo.TryAdd(RecordType, Record))
+				{
+					Record = RecordTypeToInfo[RecordType];
+				}
+			}
+			return Record;
+		}
+
+		#endregion
+
+		#region p4 add
+
+		/// <summary>
+		/// Adds files to a pending changelist.
+		/// </summary>
+		/// <param name="ChangeNumber">Changelist to add files to</param>
+		/// <param name="FileNames">Files to be added</param>
+		/// <returns>Response from the server</returns>
+		public PerforceResponseList<AddRecord> Add(int ChangeNumber, params string[] FileNames)
+		{
+			return Add(ChangeNumber, null, AddOptions.None, FileNames);
+		}
+
+		/// <summary>
+		/// Adds files to a pending changelist.
+		/// </summary>
+		/// <param name="ChangeNumber">Changelist to add files to</param>
+		/// <param name="FileType">Type for new files</param>
+		/// <param name="Options">Options for the command</param>
+		/// <param name="FileNames">Files to be added</param>
+		/// <returns>Response from the server</returns>
+		public PerforceResponseList<AddRecord> Add(int ChangeNumber, string FileType, AddOptions Options, params string[] FileNames)
+		{
+			StringBuilder Arguments = new StringBuilder("add");
+			if(ChangeNumber != -1)
+			{
+				Arguments.AppendFormat(" -c {0}", ChangeNumber);
+			}
+			if((Options & AddOptions.DowngradeToAdd) != 0)
+			{
+				Arguments.Append(" -d");
+			}
+			if((Options & AddOptions.IncludeWildcards) != 0)
+			{
+				Arguments.Append(" -f");
+			}
+			if((Options & AddOptions.NoIgnore) != 0)
+			{
+				Arguments.Append(" -I");
+			}
+			if((Options & AddOptions.PreviewOnly) != 0)
+			{
+				Arguments.Append(" -n");
+			}
+			if(FileType != null)
+			{
+				Arguments.AppendFormat(" -t \"{0}\"", FileType);
+			}
+			foreach(string FileName in FileNames)
+			{
+				Arguments.AppendFormat(" \"{0}\"", FileName);
+			}
+
+			return Command<AddRecord>(Arguments.ToString(), null);
+		}
+
+		#endregion
+
+		#region p4 change
+
+		/// <summary>
+		/// Creates a changelist with the p4 change command. 
+		/// </summary>
+		/// <param name="Record">Information for the change to create. The number field should be left set to -1.</param>
+		/// <returns>The changelist number, or an error.</returns>
+		public PerforceResponse<int> CreateChange(ChangeRecord Record)
+		{
+			if(Record.Number != -1)
+			{
+				throw new PerforceException("'Number' field should be set to -1 to create a new changelist.");
+			}
+
+			PerforceResponse Response = SingleResponseCommand("change -i", SerializeRecord(Record), null);
+			if (Response.Failed)
+			{
+				return new PerforceResponse<int>(Response.Error);
+			}
+
+			string[] Tokens = Response.Info.Data.Split(' ');
+			if (Tokens.Length != 3)
+			{
+				throw new PerforceException("Unexpected info response from change command: {0}", Response);
+			}
+
+			return new PerforceResponse<int>(int.Parse(Tokens[1]));
+		}
+
+		/// <summary>
+		/// Updates an existing changelist.
+		/// </summary>
+		/// <param name="Options">Options for this command</param>
+		/// <param name="Record">Information for the change to create. The number field should be left set to zero.</param>
+		/// <returns>The changelist number, or an error.</returns>
+		public PerforceResponse UpdateChange(UpdateChangeOptions Options, ChangeRecord Record)
+		{
+			if(Record.Number == -1)
+			{
+				throw new PerforceException("'Number' field must be set to update a changelist.");
+			}
+
+			StringBuilder Arguments = new StringBuilder("change -i");
+			if((Options & UpdateChangeOptions.Force) != 0)
+			{
+				Arguments.Append(" -f");
+			}
+			if((Options & UpdateChangeOptions.Submitted) != 0)
+			{
+				Arguments.Append(" -u");
+			}
+
+			return SingleResponseCommand(Arguments.ToString(), SerializeRecord(Record), null);
+		}
+
+		/// <summary>
+		/// Deletes a changelist (p4 change -d)
+		/// </summary>
+		/// <param name="Options">Options for the command</param>
+		/// <param name="ChangeNumber">Changelist number to delete</param>
+		/// <returns>Response from the server</returns>
+		public PerforceResponse DeleteChange(DeleteChangeOptions Options, int ChangeNumber)
+		{
+			StringBuilder Arguments = new StringBuilder("change -d");
+			if((Options & DeleteChangeOptions.Submitted) != 0)
+			{
+				Arguments.Append(" -f");
+			}
+			if((Options & DeleteChangeOptions.BeforeRenumber) != 0)
+			{
+				Arguments.Append(" -O");
+			}
+			Arguments.AppendFormat(" {0}", ChangeNumber);
+
+			return SingleResponseCommand(Arguments.ToString(), null, null);
+		}
+
+		/// <summary>
+		/// Gets a changelist
+		/// </summary>
+		/// <param name="Options">Options for the command</param>
+		/// <param name="ChangeNumber">Changelist number to retrieve. -1 is the default changelist for this workspace.</param>
+		/// <returns>Response from the server</returns>
+		public PerforceResponse<ChangeRecord> GetChange(GetChangeOptions Options, int ChangeNumber)
+		{
+			StringBuilder Arguments = new StringBuilder("change -o");
+			if((Options & GetChangeOptions.BeforeRenumber) != 0)
+			{
+				Arguments.Append(" -O");
+			}
+			Arguments.AppendFormat(" {0}", ChangeNumber);
+
+			return SingleResponseCommand<ChangeRecord>(Arguments.ToString(), null);
+		}
+
+		byte[] SerializeRecord(ChangeRecord Input)
+		{
+			List<KeyValuePair<string, object>> NameToValue = new List<KeyValuePair<string, object>>();
+			if (Input.Number == -1)
+			{
+				NameToValue.Add(new KeyValuePair<string, object>("Change", "new"));
+			}
+			else
+			{
+				NameToValue.Add(new KeyValuePair<string, object>("Change", Input.Number.ToString()));
+			}
+			if (Input.Type != ChangeType.Unspecified)
+			{
+				NameToValue.Add(new KeyValuePair<string, object>("Type", Input.Type.ToString()));
+			}
+			if (Input.User != null)
+			{
+				NameToValue.Add(new KeyValuePair<string, object>("User", Input.User));
+			}
+			if (Input.Client != null)
+			{
+				NameToValue.Add(new KeyValuePair<string, object>("Client", Input.Client));
+			}
+			if (Input.Description != null)
+			{
+				NameToValue.Add(new KeyValuePair<string, object>("Description", Input.Description));
+			}
+			return SerializeRecord(NameToValue);
+		}
+
+		#endregion
+
+		#region p4 changes
+
+		/// <summary>
+		/// Enumerates changes on the server
+		/// </summary>
+		/// <param name="Options">Options for the command</param>
+		/// <param name="MaxChanges">List only the highest numbered changes</param>
+		/// <param name="Status">Limit the list to the changelists with the given status (pending, submitted or shelved)</param>
+		/// <param name="FileSpecs">Paths to query changes for</param>
+		/// <returns>List of responses from the server.</returns>
+		public PerforceResponseList<ChangesRecord> Changes(ChangesOptions Options, int MaxChanges, ChangeStatus Status, params string[] FileSpecs)
+		{
+			return Changes(Options, null, MaxChanges, Status, null, FileSpecs);
+		}
+
+		/// <summary>
+		/// Enumerates changes on the server
+		/// </summary>
+		/// <param name="Options">Options for the command</param>
+		/// <param name="ClientName">List only changes made from the named client workspace.</param>
+		/// <param name="MaxChanges">List only the highest numbered changes</param>
+		/// <param name="Status">Limit the list to the changelists with the given status (pending, submitted or shelved)</param>
+		/// <param name="UserName">List only changes made by the named user</param>
+		/// <param name="FileSpecs">Paths to query changes for</param>
+		/// <returns>List of responses from the server.</returns>
+		public PerforceResponseList<ChangesRecord> Changes(ChangesOptions Options, string ClientName, int MaxChanges, ChangeStatus Status, string UserName, params string[] FileSpecs)
+		{
+			StringBuilder Arguments = new StringBuilder("changes");
+			if ((Options & ChangesOptions.IncludeIntegrations) != 0)
+			{
+				Arguments.Append(" -i");
+			}
+			if ((Options & ChangesOptions.IncludeTimes) != 0)
+			{
+				Arguments.Append(" -t");
+			}
+			if ((Options & ChangesOptions.LongOutput) != 0)
+			{
+				Arguments.Append(" -l");
+			}
+			if ((Options & ChangesOptions.TruncatedLongOutput) != 0)
+			{
+				Arguments.Append(" -L");
+			}
+			if ((Options & ChangesOptions.IncludeRestricted) != 0)
+			{
+				Arguments.Append(" -f");
+			}
+			if(ClientName != null)
+			{
+				Arguments.AppendFormat(" -c \"{0}\"", ClientName);
+			}
+			if(MaxChanges != -1)
+			{
+				Arguments.AppendFormat(" -m {0}", MaxChanges);
+			}
+			if(Status != ChangeStatus.All)
+			{
+				Arguments.AppendFormat(" -s {0}", GetEnumText(typeof(ChangeStatus), Status));
+			}
+			if(ClientName != null)
+			{
+				Arguments.AppendFormat(" -u {0}", UserName);
+			}
+			foreach(string FileSpec in FileSpecs)
+			{
+				Arguments.AppendFormat(" \"{0}\"", FileSpec);
+			}
+			return Command<ChangesRecord>(Arguments.ToString(), null);
+		}
+
+		#endregion
+
+		#region p4 client
+
+		/// <summary>
+		/// Creates a client
+		/// </summary>
+		/// <param name="Record">The client record</param>
+		/// <returns>Response from the server</returns>
+		public PerforceResponse CreateClient(ClientRecord Record)
+		{
+			return UpdateClient(Record);
+		}
+
+		/// <summary>
+		/// Creates a client
+		/// </summary>
+		/// <param name="Record">The client record</param>
+		/// <returns>Response from the server</returns>
+		public PerforceResponse UpdateClient(ClientRecord Record)
+		{
+			return SingleResponseCommand("client -i", SerializeRecord(Record), null);
+		}
+
+		/// <summary>
+		/// Deletes a client
+		/// </summary>
+		/// <param name="Options">Options for this command</param>
+		/// <param name="ClientName">Name of the client</param>
+		/// <returns>Response from the server</returns>
+		public PerforceResponse DeleteClient(DeleteClientOptions Options, string ClientName)
+		{
+			StringBuilder Arguments = new StringBuilder("client -d");
+			if((Options & DeleteClientOptions.Force) != 0)
+			{
+				Arguments.Append(" -f");
+			}
+			if((Options & DeleteClientOptions.DeleteShelved) != 0)
+			{
+				Arguments.Append(" -Fs");
+			}
+			Arguments.AppendFormat(" -d \"{0}\"", ClientName);
+			return SingleResponseCommand<ClientRecord>(Arguments.ToString(), null);
+		}
+
+		/// <summary>
+		/// Changes the stream associated with a client
+		/// </summary>
+		/// <param name="ClientName">The client name</param>
+		/// <param name="StreamName">The new stream to be associated with the client</param>
+		/// <param name="Options">Options for this command</param>
+		/// <returns>Response from the server</returns>
+		public PerforceResponse SwitchClientToStream(string ClientName, string StreamName, SwitchClientOptions Options)
+		{
+			StringBuilder Arguments = new StringBuilder("client -s");
+			if((Options & SwitchClientOptions.IgnoreOpenFiles) != 0)
+			{
+				Arguments.Append(" -f");
+			}
+			Arguments.AppendFormat(" -S \"{0}\"", StreamName);
+
+			return SingleResponseCommand(Arguments.ToString(), null, null);
+		}
+
+		/// <summary>
+		/// Changes a client to mirror a template
+		/// </summary>
+		/// <param name="ClientName">The client name</param>
+		/// <param name="TemplateName">The new stream to be associated with the client</param>
+		/// <returns>Response from the server</returns>
+		public PerforceResponse SwitchClientToTemplate(string ClientName, string TemplateName)
+		{
+			string Arguments = String.Format("client -s -t \"{0}\"", TemplateName);
+			return SingleResponseCommand(Arguments, null, null);
+		}
+
+		/// <summary>
+		/// Queries the current client definition
+		/// </summary>
+		/// <param name="ClientName">Name of the client. Specify null for the current client.</param>
+		/// <returns>Response from the server; either a client record or error code</returns>
+		public PerforceResponse<ClientRecord> GetClient(string ClientName)
+		{
+			StringBuilder Arguments = new StringBuilder("client -o");
+			if(ClientName != null)
+			{
+				Arguments.AppendFormat(" \"{0}\"", ClientName);
+			}
+			return SingleResponseCommand<ClientRecord>(Arguments.ToString(), null);
+		}
+
+		/// <summary>
+		/// Queries the view for a stream
+		/// </summary>
+		/// <param name="StreamName">Name of the stream.</param>
+		/// <param name="ChangeNumber">Changelist at which to query the stream view</param>
+		/// <returns>Response from the server; either a client record or error code</returns>
+		public PerforceResponse<ClientRecord> GetStreamView(string StreamName, int ChangeNumber)
+		{
+			StringBuilder Arguments = new StringBuilder("client -o");
+			Arguments.AppendFormat(" -S \"{0}\"", StreamName);
+			if(ChangeNumber != -1)
+			{
+				Arguments.AppendFormat(" -c {0}", ChangeNumber);
+			}
+			return SingleResponseCommand<ClientRecord>(Arguments.ToString(), null);
+		}
+
+		/// <summary>
+		/// Serializes a client record to a byte array
+		/// </summary>
+		/// <param name="Input">The input record</param>
+		/// <returns>Serialized record data</returns>
+		byte[] SerializeRecord(ClientRecord Input)
+		{
+			List<KeyValuePair<string, object>> NameToValue = new List<KeyValuePair<string, object>>();
+			if (Input.Name != null)
+			{
+				NameToValue.Add(new KeyValuePair<string, object>("Client", Input.Name));
+			}
+			if (Input.Owner != null)
+			{
+				NameToValue.Add(new KeyValuePair<string, object>("Owner", Input.Owner));
+			}
+			if (Input.Host != null)
+			{
+				NameToValue.Add(new KeyValuePair<string, object>("Host", Input.Host));
+			}
+			if (Input.Description != null)
+			{
+				NameToValue.Add(new KeyValuePair<string, object>("Description", Input.Description));
+			}
+			if (Input.Root != null)
+			{
+				NameToValue.Add(new KeyValuePair<string, object>("Root", Input.Root));
+			}
+			if(Input.Options != ClientOptions.None)
+			{
+				NameToValue.Add(new KeyValuePair<string, object>("Options", GetEnumText(typeof(ClientOptions), Input.Options)));
+			}
+			if(Input.SubmitOptions != ClientSubmitOptions.Unspecified)
+			{
+				NameToValue.Add(new KeyValuePair<string, object>("SubmitOptions", GetEnumText(typeof(ClientSubmitOptions), Input.SubmitOptions)));
+			}
+			if(Input.LineEnd != ClientLineEndings.Unspecified)
+			{
+				NameToValue.Add(new KeyValuePair<string, object>("LineEnd", GetEnumText(typeof(ClientLineEndings), Input.LineEnd)));
+			}
+			if(Input.Stream != null)
+			{
+				NameToValue.Add(new KeyValuePair<string, object>("Stream", Input.Stream));
+			}
+			for (int Idx = 0; Idx < Input.View.Count; Idx++)
+			{
+				NameToValue.Add(new KeyValuePair<string, object>(String.Format("View{0}", Idx), Input.View[Idx]));
+			}
+			return SerializeRecord(NameToValue);
+		}
+
+		#endregion
+
+		#region p4 clients
+
+		/// <summary>
+		/// Queries the current client definition
+		/// </summary>
+		/// <param name="Options">Options for this command</param>
+		/// <param name="UserName">List only client workspaces owned by this user.</param>
+		/// <returns>Response from the server; either a client record or error code</returns>
+		public PerforceResponseList<ClientsRecord> Clients(ClientsOptions Options, string UserName)
+		{
+			return Clients(Options, null, -1, null, UserName);
+		}
+
+		/// <summary>
+		/// Queries the current client definition
+		/// </summary>
+		/// <param name="Options">Options for this command</param>
+		/// <param name="Filter">List only client workspaces matching filter. Treated as case sensitive if <ref>ClientsOptions.CaseSensitiveFilter</ref> is set.</param>
+		/// <param name="MaxResults">Limit the number of results to return. -1 for all.</param>
+		/// <param name="Stream">List client workspaces associated with the specified stream.</param>
+		/// <param name="UserName">List only client workspaces owned by this user.</param>
+		/// <returns>Response from the server; either a client record or error code</returns>
+		public PerforceResponseList<ClientsRecord> Clients(ClientsOptions Options, string Filter, int MaxResults, string Stream, string UserName)
+		{
+			StringBuilder Arguments = new StringBuilder("clients");
+			if((Options & ClientsOptions.All) != 0)
+			{
+				Arguments.Append(" -a");
+			}
+			if (Filter != null)
+			{
+				if ((Options & ClientsOptions.CaseSensitiveFilter) != 0)
+				{
+					Arguments.AppendFormat(" -e \"{0}\"", Filter);
+				}
+				else
+				{
+					Arguments.AppendFormat(" -E \"{0}\"", Filter);
+				}
+			}
+			if(MaxResults != -1)
+			{
+				Arguments.AppendFormat(" -m {0}", MaxResults);
+			}
+			if(Stream != null)
+			{
+				Arguments.AppendFormat(" -S \"{0}\"", Stream);
+			}
+			if((Options & ClientsOptions.WithTimes) != 0)
+			{
+				Arguments.Append(" -t");
+			}
+			if(UserName != null)
+			{
+				Arguments.AppendFormat(" -u \"{0}\"", UserName);
+			}
+			if((Options & ClientsOptions.Unloaded) != 0)
+			{
+				Arguments.Append(" -U");
+			}
+			return Command<ClientsRecord>(Arguments.ToString(), null);
+		}
+
+		#endregion
+
+		#region p4 describe
+
+		/// <summary>
+		/// Describes a single changelist
+		/// </summary>
+		/// <param name="ChangeNumber">The changelist number to retrieve description for</param>
+		/// <returns>Response from the server; either a describe record or error code</returns>
+		public PerforceResponse<DescribeRecord> Describe(int ChangeNumber)
+		{
+			PerforceResponseList<DescribeRecord> Records = Describe(new int[] { ChangeNumber });
+			if (Records.Count != 1)
+			{
+				throw new PerforceException("Expected only one record returned from p4 describe command, got {0}", Records.Count);
+			}
+			return Records[0];
+		}
+
+		/// <summary>
+		/// Describes a set of changelists
+		/// </summary>
+		/// <param name="ChangeNumbers">The changelist numbers to retrieve descriptions for</param>
+		/// <returns>List of responses from the server</returns>
+		public PerforceResponseList<DescribeRecord> Describe(params int[] ChangeNumbers)
+		{
+			StringBuilder Arguments = new StringBuilder("describe -s");
+			foreach(int ChangeNumber in ChangeNumbers)
+			{
+				Arguments.AppendFormat(" {0}", ChangeNumber);
+			}
+			return Command<DescribeRecord>(Arguments.ToString(), null);
+		}
+
+		#endregion
+
+		#region p4 edit
+
+		/// <summary>
+		/// Opens files for edit
+		/// </summary>
+		/// <param name="ChangeNumber">Changelist to add files to</param>
+		/// <param name="FileSpecs">Files to be opened for edit</param>
+		/// <returns>Response from the server</returns>
+		public PerforceResponseList<EditRecord> Edit(int ChangeNumber, params string[] FileSpecs)
+		{
+			return Edit(ChangeNumber, null, EditOptions.None, FileSpecs);
+		}
+
+		/// <summary>
+		/// Opens files for edit
+		/// </summary>
+		/// <param name="ChangeNumber">Changelist to add files to</param>
+		/// <param name="FileType">Type for new files</param>
+		/// <param name="Options">Options for the command</param>
+		/// <param name="FileSpecs">Files to be opened for edit</param>
+		/// <returns>Response from the server</returns>
+		public PerforceResponseList<EditRecord> Edit(int ChangeNumber, string FileType, EditOptions Options, params string[] FileSpecs)
+		{
+			StringBuilder Arguments = new StringBuilder("edit");
+			if(ChangeNumber != -1)
+			{
+				Arguments.AppendFormat(" -c {0}", ChangeNumber);
+			}
+			if((Options & EditOptions.KeepWorkspaceFiles) != 0)
+			{
+				Arguments.Append(" -k");
+			}
+			if((Options & EditOptions.PreviewOnly) != 0)
+			{
+				Arguments.Append(" -n");
+			}
+			if(FileType != null)
+			{
+				Arguments.AppendFormat(" -t \"{0}\"", FileType);
+			}
+			foreach(string FileSpec in FileSpecs)
+			{
+				Arguments.AppendFormat(" \"{0}\"", FileSpec);
+			}
+
+			return Command<EditRecord>(Arguments.ToString(), null);
+		}
+
+		#endregion
+
+		#region p4 filelog
+
+		/// <summary>
+		/// Execute the 'filelog' command
+		/// </summary>
+		/// <param name="Options">Options for the command</param>
+		/// <param name="FileSpecs">List of file specifications to query</param>
+		/// <returns>List of responses from the server</returns>
+		public PerforceResponseList<FileLogRecord> FileLog(FileLogOptions Options, params string[] FileSpecs)
+		{
+			return FileLog(-1, -1, Options, FileSpecs);
+		}
+
+		/// <summary>
+		/// Execute the 'filelog' command
+		/// </summary>
+		/// <param name="MaxChanges">Number of changelists to show. Ignored if zero or negative.</param>
+		/// <param name="Options">Options for the command</param>
+		/// <param name="FileSpecs">List of file specifications to query</param>
+		/// <returns>List of responses from the server</returns>
+		public PerforceResponseList<FileLogRecord> FileLog(int MaxChanges, FileLogOptions Options, params string[] FileSpecs)
+		{
+			return FileLog(-1, MaxChanges, Options, FileSpecs);
+		}
+
+		/// <summary>
+		/// Execute the 'filelog' command
+		/// </summary>
+		/// <param name="ChangeNumber">Show only files modified by this changelist. Ignored if zero or negative.</param>
+		/// <param name="MaxChanges">Number of changelists to show. Ignored if zero or negative.</param>
+		/// <param name="Options">Options for the command</param>
+		/// <param name="FileSpecs">List of file specifications to query</param>
+		/// <returns>List of responses from the server</returns>
+		public PerforceResponseList<FileLogRecord> FileLog(int ChangeNumber, int MaxChanges, FileLogOptions Options, params string[] FileSpecs)
+		{
+			// Build the argument list
+			StringBuilder Arguments = new StringBuilder("filelog");
+			if(ChangeNumber > 0)
+			{
+				Arguments.AppendFormat(" -c {0}", ChangeNumber);
+			}
+			if((Options & FileLogOptions.ContentHistory) != 0)
+			{
+				Arguments.Append(" -h");
+			}
+			if((Options & FileLogOptions.FollowAcrossBranches) != 0)
+			{
+				Arguments.Append(" -i");
+			}
+			if((Options & FileLogOptions.FullDescriptions) != 0)
+			{
+				Arguments.Append(" -l");
+			}
+			if((Options & FileLogOptions.LongDescriptions) != 0)
+			{
+				Arguments.Append(" -L");
+			}
+			if(MaxChanges > 0)
+			{
+				Arguments.AppendFormat(" -m {0}", MaxChanges);
+			}
+			if((Options & FileLogOptions.DoNotFollowPromotedTaskStreams) != 0)
+			{
+				Arguments.Append(" -p");
+			}
+			if((Options & FileLogOptions.IgnoreNonContributoryIntegrations) != 0)
+			{
+				Arguments.Append(" -s");
+			}
+
+			// Always include times to simplify parsing
+			Arguments.Append(" -t");
+
+			// Add the file arguments
+			foreach(string FileSpec in FileSpecs)
+			{
+				Arguments.AppendFormat(" \"{0}\"", FileSpec);
+			}
+
+			// Execute the command
+			return Command<FileLogRecord>(Arguments.ToString(), null);
+		}
+
+		#endregion
+
+		#region p4 fstat
+
+		/// <summary>
+		/// Execute the 'fstat' command
+		/// </summary>
+		/// <param name="Options">Options for the command</param>
+		/// <param name="FileSpecs">List of file specifications to query</param>
+		/// <returns>List of responses from the server</returns>
+		public PerforceResponseList<FStatRecord> FStat(FStatOptions Options, params string[] FileSpecs)
+		{
+			return FStat(-1, Options, FileSpecs);
+		}
+
+		/// <summary>
+		/// Execute the 'fstat' command
+		/// </summary>
+		/// <param name="MaxFiles">Produce fstat output for only the first max files.</param>
+		/// <param name="Options">Options for the command</param>
+		/// <param name="FileSpecs">List of file specifications to query</param>
+		/// <returns>List of responses from the server</returns>
+		public PerforceResponseList<FStatRecord> FStat(int MaxFiles, FStatOptions Options, params string[] FileSpecs)
+		{
+			return FStat(-1, -1, null, MaxFiles, Options, FileSpecs);
+		}
+
+		/// <summary>
+		/// Execute the 'fstat' command
+		/// </summary>
+		/// <param name="AfterChangeNumber">Return only files affected after the given changelist number.</param>
+		/// <param name="OnlyChangeNumber">Return only files affected by the given changelist number.</param>
+		/// <param name="Filter">List only those files that match the criteria specified.</param>
+		/// <param name="MaxFiles">Produce fstat output for only the first max files.</param>
+		/// <param name="Options">Options for the command</param>
+		/// <param name="FileSpecs">List of file specifications to query</param>
+		/// <returns>List of responses from the server</returns>
+		public PerforceResponseList<FStatRecord> FStat(int AfterChangeNumber, int OnlyChangeNumber, string Filter, int MaxFiles, FStatOptions Options, params string[] FileSpecs)
+		{
+			// Build the argument list
+			StringBuilder Arguments = new StringBuilder("fstat");
+			if (AfterChangeNumber != -1)
+			{
+				Arguments.AppendFormat(" -c {0}", AfterChangeNumber);
+			}
+			if(OnlyChangeNumber != -1)
+			{
+				Arguments.AppendFormat(" -e {0}", OnlyChangeNumber);
+			}
+			if(Filter != null)
+			{
+				Arguments.AppendFormat(" -F \"{0}\"", Filter);
+			}
+			if((Options & FStatOptions.ReportDepotSyntax) != 0)
+			{
+				Arguments.Append(" -L");
+			}
+			if((Options & FStatOptions.AllRevisions) != 0)
+			{
+				Arguments.Append(" -Of");
+			}
+			if((Options & FStatOptions.IncludeFileSizes) != 0)
+			{
+				Arguments.Append(" -Ol");
+			}
+			if((Options & FStatOptions.ClientFileInPerforceSyntax) != 0)
+			{
+				Arguments.Append(" -Op");
+			}
+			if((Options & FStatOptions.ShowPendingIntegrations) != 0)
+			{
+				Arguments.Append(" -Or");
+			}
+			if((Options & FStatOptions.ShortenOutput) != 0)
+			{
+				Arguments.Append(" -Os");
+			}
+			if((Options & FStatOptions.ReverseOrder) != 0)
+			{
+				Arguments.Append(" -r");
+			}
+			if((Options & FStatOptions.OnlyMapped) != 0)
+			{
+				Arguments.Append(" -Rc");
+			}
+			if((Options & FStatOptions.OnlyHave) != 0)
+			{
+				Arguments.Append(" -Rh");
+			}
+			if((Options & FStatOptions.OnlyOpenedBeforeHead) != 0)
+			{
+				Arguments.Append(" -Rn");
+			}
+			if((Options & FStatOptions.OnlyOpenInWorkspace) != 0)
+			{
+				Arguments.Append(" -Ro");
+			}
+			if((Options & FStatOptions.OnlyOpenAndResolved) != 0)
+			{
+				Arguments.Append(" -Rr");
+			}
+			if((Options & FStatOptions.OnlyShelved) != 0)
+			{
+				Arguments.Append(" -Rs");
+			}
+			if((Options & FStatOptions.OnlyUnresolved) != 0)
+			{
+				Arguments.Append(" -Ru");
+			}
+			if((Options & FStatOptions.SortByDate) != 0)
+			{
+				Arguments.Append(" -Sd");
+			}
+			if((Options & FStatOptions.SortByHaveRevision) != 0)
+			{
+				Arguments.Append(" -Sh");
+			}
+			if((Options & FStatOptions.SortByHeadRevision) != 0)
+			{
+				Arguments.Append(" -Sr");
+			}
+			if((Options & FStatOptions.SortByFileSize) != 0)
+			{
+				Arguments.Append(" -Ss");
+			}
+			if((Options & FStatOptions.SortByFileType) != 0)
+			{
+				Arguments.Append(" -St");
+			}
+			if((Options & FStatOptions.IncludeFilesInUnloadDepot) != 0)
+			{
+				Arguments.Append(" -U");
+			}
+
+			// Add the file arguments
+			foreach(string FileSpec in FileSpecs)
+			{
+				Arguments.AppendFormat(" \"{0}\"", FileSpec);
+			}
+
+			// Execute the command
+			return Command<FStatRecord>(Arguments.ToString(), null);
+		}
+
+		#endregion
+
+		#region p4 info
+
+		/// <summary>
+		/// Execute the 'info' command
+		/// </summary>
+		/// <param name="Options">Options for the command</param>
+		/// <returns>Response from the server; an InfoRecord or error code</returns>
+		public PerforceResponse<InfoRecord> Info(InfoOptions Options)
+		{
+			// Build the argument list
+			StringBuilder Arguments = new StringBuilder("info");
+			if((Options & InfoOptions.ShortOutput) != 0)
+			{
+				Arguments.Append(" -s");
+			}
+			return SingleResponseCommand<InfoRecord>(Arguments.ToString(), null);
+		}
+
+		#endregion
+
+		#region p4 print
+
+		/// <summary>
+		/// Execute the 'print' command
+		/// </summary>
+		/// <param name="OutputFile">Output file to redirect output to</param>
+		/// <param name="FileSpec">Specification for the files to print</param>
+		/// <returns>Response from the server</returns>
+		public PerforceResponse<PrintRecord> Print(string OutputFile, string FileSpec)
+		{
+			// Build the argument list
+			StringBuilder Arguments = new StringBuilder("print");
+			Arguments.AppendFormat(" -o \"{0}\"", OutputFile);
+			Arguments.AppendFormat(" \"{0}\"", FileSpec);
+			return SingleResponseCommand<PrintRecord>(Arguments.ToString(), null);
+		}
+
+		class PrintHandler : IDisposable
+		{
+			Dictionary<string, string> DepotFileToLocalFile;
+			FileStream OutputStream;
+
+			public PrintHandler(Dictionary<string, string> DepotFileToLocalFile)
+			{
+				this.DepotFileToLocalFile = DepotFileToLocalFile;
+			}
+
+			public void Dispose()
+			{
+				CloseStream();
+			}
+
+			private void OpenStream(string FileName)
+			{
+				CloseStream();
+				Directory.CreateDirectory(Path.GetDirectoryName(FileName));
+				OutputStream = File.Open(FileName, FileMode.Create, FileAccess.Write, FileShare.None);
+			}
+
+			private void CloseStream()
+			{
+				if(OutputStream != null)
+				{
+					OutputStream.Dispose();
+					OutputStream = null;
+				}
+			}
+
+			public void HandleRecord(List<KeyValuePair<string, object>> Fields)
+			{
+				if(Fields[0].Key != "code")
+				{
+					throw new Exception("Missing code field");
+				}
+
+				string Value = (string)Fields[0].Value;
+				if(Value == "stat")
+				{
+					string DepotFile = Fields.First(x => x.Key == "depotFile").Value.ToString();
+
+					string LocalFile;
+					if(!DepotFileToLocalFile.TryGetValue(DepotFile, out LocalFile))
+					{
+						throw new PerforceException("Depot file '{0}' not found in input dictionary", DepotFile);
+					}
+
+					OpenStream(LocalFile);
+				}
+				else if(Value == "binary" || Value == "text")
+				{
+					byte[] Data = (byte[])Fields.First(x => x.Key == "data").Value;
+					OutputStream.Write(Data, 0, Data.Length);
+				}
+				else
+				{
+					throw new Exception("Unexpected record type");
+				}
+			}
+		}
+
+		/// <summary>
+		/// Execute the 'print' command
+		/// </summary>
+		/// <param name="OutputFile">Output file to redirect output to</param>
+		/// <param name="FileSpec">Specification for the files to print</param>
+		/// <returns>Response from the server</returns>
+		public void Print(Dictionary<string, string> DepotFileToLocalFile)
+		{
+			string ListFileName = Path.GetTempFileName();
+			try
+			{
+				// Write the list of depot files
+				File.WriteAllLines(ListFileName, DepotFileToLocalFile.Keys);
+
+				// Execute Perforce, consuming the binary output into a memory stream
+				using(PrintHandler Handler = new PrintHandler(DepotFileToLocalFile))
+				{
+					Command(String.Format("-x \"{0}\" print", ListFileName), null, Handler.HandleRecord);
+				}
+			}
+			finally
+			{
+				File.Delete(ListFileName);
+			}
+
+/*
+			using(PerforceChildProcess ChildProcess = new PerforceChildProcess("{0} print \"{1}\"", GlobalOptions, FileSpec))
+			{
+				PerforceRecordType Type;
+				if(ChildProcess.TryBeginRecord(out Type))
+				{
+					for(;;)
+					{
+						if(Type != PerforceRecordType.Stat)
+						{
+							throw new PerforceException("Unexpected record type: {0}", Type);
+						}
+
+						Dictionary<string, object> Fields = new Dictionary<string, object>();
+						ChildProcess.ReadFields(Fields);
+
+						string OutputFile = Path.Combine("D:\\WSTEST\\Content", Fields["depotFile"].ToString().Replace('/', '+'));
+//						Console.WriteLine(OutputFile);
+
+						if(File.Exists(OutputFile))
+						{
+							File.SetAttributes(OutputFile, FileAttributes.Normal);
+						}
+
+						using(FileStream OutputStream = File.Open(OutputFile, FileMode.Create, FileAccess.Write, FileShare.Read))
+						{
+							for(;;)
+							{
+								if(!ChildProcess.TryBeginRecord(out Type))
+								{
+									return;
+								}
+								if(Type != PerforceRecordType.Text && Type != PerforceRecordType.Binary)
+								{
+									break;
+								}
+
+								string Key;
+								if(!ChildProcess.TryReadKey(out Key))
+								{
+									throw new PerforceException("Foo");
+								}
+								if(Key != "data")
+								{
+									throw new PerforceException("Bar");
+								}
+								ChildProcess.ReadDataValue(OutputStream);
+
+								ChildProcess.ReadFields(Fields);
+							}
+						}
+					}
+				}
+			}*/
+/*				while(!Process.IsEndOfFile())
+				{
+					// Check that a dictionary follows
+					byte Temp = Process.ReadByte();
+					if(Temp != '{')
+					{
+						throw new PerforceException("Unexpected data while parsing marshalled output - expected '{'");
+					}
+
+					// Read all the fields in the record
+					for(;;)
+					{
+						// Read the next field type. Perforce only outputs string records. A '0' character indicates the end of the dictionary.
+						byte KeyFieldType = Process.ReadByte();
+						if(KeyFieldType == '0')
+						{
+							break;
+						}
+						else if(KeyFieldType != 's')
+						{
+							throw new PerforceException("Unexpected key field type while parsing marshalled output ({0}) - expected 's'", (int)KeyFieldType);
+						}
+
+						// Read the key
+						string Key = Process.ReadString();
+
+						// Read the value type.
+						byte ValueFieldType = Process.ReadByte();
+						if(ValueFieldType == 'i')
+						{
+							// An integer
+							string Value = Process.ReadInt32().ToString();
+							Console.WriteLine("{0}={1}", Key, Value);
+//							Record.Add(new KeyValuePair<string, object>(Key, Value));
+						}
+						else if(ValueFieldType == 's')
+						{
+							// A string (or variable length data)
+							if(Key == "data")
+							{
+								MemoryStream OutputData = new MemoryStream();
+								Process.CopyData(OutputData);
+								Console.WriteLine("{0}={1} bytes", Key, OutputData.Length);
+							}
+							else
+							{
+								string Value = Process.ReadString();
+								Console.WriteLine("{0}={1}", Key, Value);
+							}
+						}
+						else
+						{
+							throw new PerforceException("Unexpected value field type while parsing marshalled output ({0}) - expected 's'", (int)ValueFieldType);
+						}
+					}
+				}
+			}*/
+		}
+
+		#endregion
+
+		#region p4 reconcile
+		
+		/// <summary>
+		/// Open files for add, delete, and/or edit in order to reconcile a workspace with changes made outside of Perforce.
+		/// </summary>
+		/// <param name="ChangeNumber">Changelist to open files to</param>
+		/// <param name="Options">Options for the command</param>
+		/// <param name="FileSpecs">Files to be reverted</param>
+		/// <returns>Response from the server</returns>
+		public PerforceResponseList<ReconcileRecord> Reconcile(int ChangeNumber, ReconcileOptions Options, params string[] FileSpecs)
+		{
+			StringBuilder Arguments = new StringBuilder("reconcile");
+			if(ChangeNumber != -1)
+			{
+				Arguments.AppendFormat(" -c {0}", ChangeNumber);
+			}
+			if((Options & ReconcileOptions.Edit) != 0)
+			{
+				Arguments.Append(" -e");
+			}
+			if((Options & ReconcileOptions.Add) != 0)
+			{
+				Arguments.Append(" -a");
+			}
+			if((Options & ReconcileOptions.Delete) != 0)
+			{
+				Arguments.Append(" -d");
+			}
+			if((Options & ReconcileOptions.PreviewOnly) != 0)
+			{
+				Arguments.Append(" -n");
+			}
+			if((Options & ReconcileOptions.AllowWildcards) != 0)
+			{
+				Arguments.Append(" -f");
+			}
+			if((Options & ReconcileOptions.NoIgnore) != 0)
+			{
+				Arguments.Append(" -I");
+			}
+			if((Options & ReconcileOptions.LocalFileSyntax) != 0)
+			{
+				Arguments.Append(" -l");
+			}
+			foreach(string FileSpec in FileSpecs)
+			{
+				Arguments.AppendFormat(" \"{0}\"", FileSpec);
+			}
+
+			return Command<ReconcileRecord>(Arguments.ToString(), null);
+		}
+
+		#endregion
+
+		#region p4 resolve
+
+		/// <summary>
+		/// Resolve conflicts between file revisions.
+		/// </summary>
+		/// <param name="ChangeNumber">Changelist to open files to</param>
+		/// <param name="Options">Options for the command</param>
+		/// <param name="FileSpecs">Files to be reverted</param>
+		/// <returns>Response from the server</returns>
+		public PerforceResponseList<ResolveRecord> Resolve(int ChangeNumber, ResolveOptions Options, params string[] FileSpecs)
+		{
+			StringBuilder Arguments = new StringBuilder("resolve");
+			if ((Options & ResolveOptions.Automatic) != 0)
+			{
+				Arguments.Append(" -am");
+			}
+			if ((Options & ResolveOptions.AcceptYours) != 0)
+			{
+				Arguments.Append(" -ay");
+			}
+			if ((Options & ResolveOptions.AcceptTheirs) != 0)
+			{
+				Arguments.Append(" -at");
+			}
+			if ((Options & ResolveOptions.SafeAccept) != 0)
+			{
+				Arguments.Append(" -as");
+			}
+			if ((Options & ResolveOptions.ForceAccept) != 0)
+			{
+				Arguments.Append(" -af");
+			}
+			if((Options & ResolveOptions.IgnoreWhitespaceOnly) != 0)
+			{
+				Arguments.Append(" -db");
+			}
+			if((Options & ResolveOptions.IgnoreWhitespace) != 0)
+			{
+				Arguments.Append(" -dw");
+			}
+			if((Options & ResolveOptions.IgnoreLineEndings) != 0)
+			{
+				Arguments.Append(" -dl");
+			}
+			if ((Options & ResolveOptions.ResolveAgain) != 0)
+			{
+				Arguments.Append(" -f");
+			}
+			if ((Options & ResolveOptions.PreviewOnly) != 0)
+			{
+				Arguments.Append(" -n");
+			}
+			foreach(string FileSpec in FileSpecs)
+			{
+				Arguments.AppendFormat(" \"{0}\"", FileSpec);
+			}
+
+			return Command<ResolveRecord>(Arguments.ToString(), null);
+		}
+
+		#endregion
+
+		#region p4 revert
+
+		/// <summary>
+		/// Reverts files that have been added to a pending changelist.
+		/// </summary>
+		/// <param name="ChangeNumber">Changelist to add files to</param>
+		/// <param name="ClientName">Revert another user’s open files. </param>
+		/// <param name="Options">Options for the command</param>
+		/// <param name="FileSpecs">Files to be reverted</param>
+		/// <returns>Response from the server</returns>
+		public PerforceResponseList<RevertRecord> Revert(int ChangeNumber, string ClientName, RevertOptions Options, params string[] FileSpecs)
+		{
+			StringBuilder Arguments = new StringBuilder("revert");
+			if((Options & RevertOptions.Unchanged) != 0)
+			{
+				Arguments.Append(" -a");
+			}
+			if((Options & RevertOptions.PreviewOnly) != 0)
+			{
+				Arguments.Append(" -n");
+			}
+			if((Options & RevertOptions.KeepWorkspaceFiles) != 0)
+			{
+				Arguments.Append(" -k");
+			}
+			if((Options & RevertOptions.DeleteAddedFiles) != 0)
+			{
+				Arguments.Append(" -w");
+			}
+			if(ChangeNumber != -1)
+			{
+				Arguments.AppendFormat(" -c {0}", ChangeNumber);
+			}
+			if(ClientName != null)
+			{
+				Arguments.AppendFormat(" -C \"{0}\"", ClientName);
+			}
+			foreach(string FileSpec in FileSpecs)
+			{
+				Arguments.AppendFormat(" \"{0}\"", FileSpec);
+			}
+
+			return Command<RevertRecord>(Arguments.ToString(), null);
+		}
+
+		#endregion
+
+		#region p4 shelve
+
+		/// <summary>
+		/// Shelves a set of files
+		/// </summary>
+		/// <param name="ChangeNumber">The change number to receive the shelved files</param>
+		/// <param name="Options">Options for the command</param>
+		/// <param name="FileSpecs">Files to sync</param>
+		/// <returns>Response from the server</returns>
+		public PerforceResponseList<ShelveRecord> Shelve(int ChangeNumber, ShelveOptions Options, params string[] FileSpecs)
+		{
+			StringBuilder Arguments = new StringBuilder("shelve");
+			Arguments.AppendFormat(" -c {0}", ChangeNumber);
+			if((Options & ShelveOptions.OnlyChanged) != 0)
+			{
+				Arguments.Append(" -a leaveunchanged");
+			}
+			if((Options & ShelveOptions.Overwrite) != 0)
+			{
+				Arguments.Append(" -f");
+			}
+			foreach(string FileSpec in FileSpecs)
+			{
+				Arguments.AppendFormat(" \"{0}\"", FileSpec);
+			}
+
+			return Command<ShelveRecord>(Arguments.ToString(), null);
+		}
+
+		/// <summary>
+		/// Deletes files from a shelved changelist
+		/// </summary>
+		/// <param name="ChangeNumber">Changelist containing shelved files to be deleted</param>
+		/// <param name="FileSpecs">Files to delete</param>
+		/// <returns>Response from the server</returns>
+		public PerforceResponse DeleteShelvedFiles(int ChangeNumber, params string[] FileSpecs)
+		{
+			StringBuilder Arguments = new StringBuilder("shelve -d");
+			if(ChangeNumber != -1)
+			{
+				Arguments.AppendFormat(" -c {0}", ChangeNumber);
+			}
+			foreach (string FileSpec in FileSpecs)
+			{
+				Arguments.AppendFormat(" \"{0}\"", FileSpec);
+			}
+
+			return SingleResponseCommand(Arguments.ToString(), null, null);
+		}
+
+		#endregion
+
+		#region p4 streams
+
+		/// <summary>
+		/// Enumerates all streams in a depot
+		/// </summary>
+		/// <param name="StreamPath">The path for streams to enumerate (eg. "//UE4/...")</param>
+		/// <returns>List of streams matching the given criteria</returns>
+		public PerforceResponseList<StreamRecord> Streams(string StreamPath)
+		{
+			return Streams(StreamPath, -1, null, false);
+		}
+
+		/// <summary>
+		/// Enumerates all streams in a depot
+		/// </summary>
+		/// <param name="StreamPath">The path for streams to enumerate (eg. "//UE4/...")</param>
+		/// <param name="MaxResults">Maximum number of results to return</param>
+		/// <param name="Filter">Additional filter to be applied to the results</param>
+		/// <param name="bUnloaded">Whether to enumerate unloaded workspaces</param>
+		/// <returns>List of streams matching the given criteria</returns>
+		public PerforceResponseList<StreamRecord> Streams(string StreamPath, int MaxResults, string Filter, bool bUnloaded)
+		{
+			// Build the command line
+			StringBuilder Arguments = new StringBuilder("streams");
+			if (bUnloaded)
+			{
+				Arguments.Append(" -U");
+			}
+			if (Filter != null)
+			{
+				Arguments.AppendFormat("-F \"{0}\"", Filter);
+			}
+			if (MaxResults > 0)
+			{
+				Arguments.AppendFormat("-m {0}", MaxResults);
+			}
+			Arguments.AppendFormat(" \"{0}\"", StreamPath);
+
+			// Execute the command
+			return Command<StreamRecord>(Arguments.ToString(), null);
+		}
+
+		#endregion
+
+		#region p4 submit
+
+		/// <summary>
+		/// Submits a pending changelist
+		/// </summary>
+		/// <param name="ChangeNumber">The changelist to submit</param>
+		/// <param name="Options">Options for the command</param>
+		/// <returns>Response from the server</returns>
+		public PerforceResponse<SubmitRecord> Submit(int ChangeNumber, SubmitOptions Options)
+		{
+			StringBuilder Arguments = new StringBuilder("submit");
+			if((Options & SubmitOptions.ReopenAsEdit) != 0)
+			{
+				Arguments.Append(" -r");
+			}
+			Arguments.AppendFormat(" -c {0}", ChangeNumber);
+
+			return Command<SubmitRecord>(Arguments.ToString(), null)[0];
+		}
+
+		#endregion
+
+		#region p4 sync
+
+		/// <summary>
+		/// Syncs files from the server
+		/// </summary>
+		/// <param name="FileSpecs">Files to sync</param>
+		/// <returns>Response from the server</returns>
+		public PerforceResponseList<SyncRecord> Sync(params string[] FileSpecs)
+		{
+			return Sync(SyncOptions.None, -1, FileSpecs);
+		}
+
+		/// <summary>
+		/// Syncs files from the server
+		/// </summary>
+		/// <param name="Options">Options for the command</param>
+		/// <param name="MaxFiles">Syncs only the first number of files specified.</param>
+		/// <param name="FileSpecs">Files to sync</param>
+		/// <returns>Response from the server</returns>
+		public PerforceResponseList<SyncRecord> Sync(SyncOptions Options, int MaxFiles, params string[] FileSpecs)
+		{
+			return Sync(Options, -1, -1, -1, -1, -1, -1, FileSpecs);
+		}
+
+		/// <summary>
+		/// Syncs files from the server
+		/// </summary>
+		/// <param name="Options">Options for the command</param>
+		/// <param name="MaxFiles">Syncs only the first number of files specified</param>
+		/// <param name="NumThreads">Sync in parallel using the given number of threads</param>
+		/// <param name="Batch">The number of files in a batch</param>
+		/// <param name="BatchSize">The number of bytes in a batch</param>
+		/// <param name="Min">Minimum number of files in a parallel sync</param>
+		/// <param name="MinSize">Minimum number of bytes in a parallel sync</param>
+		/// <param name="FileSpecs">Files to sync</param>
+		/// <returns>Response from the server</returns>
+		public PerforceResponseList<SyncRecord> Sync(SyncOptions Options, int MaxFiles, int NumThreads, int Batch, int BatchSize, int Min, int MinSize, params string[] FileSpecs)
+		{
+			string Arguments = GetSyncArguments(Options, MaxFiles, NumThreads, Batch, BatchSize, Min, MinSize, FileSpecs, false);
+			return Command<SyncRecord>(Arguments.ToString(), null);
+		}
+
+		/// <summary>
+		/// Syncs files from the server
+		/// </summary>
+		/// <param name="Options">Options for the command</param>
+		/// <param name="MaxFiles">Syncs only the first number of files specified.</param>
+		/// <param name="FileSpecs">Files to sync</param>
+		/// <returns>Response from the server</returns>
+		public PerforceResponseList<SyncSummaryRecord> SyncQuiet(SyncOptions Options, int MaxFiles, params string[] FileSpecs)
+		{
+			return SyncQuiet(Options, -1, -1, -1, -1, -1, -1, FileSpecs);
+		}
+
+		/// <summary>
+		/// Syncs files from the server without returning detailed file info
+		/// </summary>
+		/// <param name="Options">Options for the command</param>
+		/// <param name="MaxFiles">Syncs only the first number of files specified</param>
+		/// <param name="NumThreads">Sync in parallel using the given number of threads</param>
+		/// <param name="Batch">The number of files in a batch</param>
+		/// <param name="BatchSize">The number of bytes in a batch</param>
+		/// <param name="Min">Minimum number of files in a parallel sync</param>
+		/// <param name="MinSize">Minimum number of bytes in a parallel sync</param>
+		/// <param name="FileSpecs">Files to sync</param>
+		/// <returns>Response from the server</returns>
+		public PerforceResponseList<SyncSummaryRecord> SyncQuiet(SyncOptions Options, int MaxFiles, int NumThreads, int Batch, int BatchSize, int Min, int MinSize, params string[] FileSpecs)
+		{
+			string Arguments = GetSyncArguments(Options, MaxFiles, NumThreads, Batch, BatchSize, Min, MinSize, FileSpecs, true);
+			return Command<SyncSummaryRecord>(Arguments.ToString(), null);
+		}
+
+		/// <summary>
+		/// Gets arguments for a sync command
+		/// </summary>
+		/// <param name="Options">Options for the command</param>
+		/// <param name="MaxFiles">Syncs only the first number of files specified</param>
+		/// <param name="NumThreads">Sync in parallel using the given number of threads</param>
+		/// <param name="Batch">The number of files in a batch</param>
+		/// <param name="BatchSize">The number of bytes in a batch</param>
+		/// <param name="Min">Minimum number of files in a parallel sync</param>
+		/// <param name="MinSize">Minimum number of bytes in a parallel sync</param>
+		/// <param name="FileSpecs">Files to sync</param>
+		/// <param name="bQuiet">Whether to use quiet output</param>
+		/// <returns>Arguments for the command</returns>
+		static string GetSyncArguments(SyncOptions Options, int MaxFiles, int NumThreads, int Batch, int BatchSize, int Min, int MinSize, string[] FileSpecs, bool bQuiet)
+		{
+			StringBuilder Arguments = new StringBuilder("sync");
+			if((Options & SyncOptions.Force) != 0)
+			{
+				Arguments.Append(" -f");
+			}
+			if((Options & SyncOptions.KeepWorkspaceFiles) != 0)
+			{
+				Arguments.Append(" -k");
+			}
+			if((Options & SyncOptions.FullDepotSyntax) != 0)
+			{
+				Arguments.Append(" -L");
+			}
+			if((Options & SyncOptions.PreviewOnly) != 0)
+			{
+				Arguments.Append(" -n");
+			}
+			if((Options & SyncOptions.NetworkPreviewOnly) != 0)
+			{
+				Arguments.Append(" -N");
+			}
+			if((Options & SyncOptions.DoNotUpdateHaveList) != 0)
+			{
+				Arguments.Append(" -p");
+			}
+			if(bQuiet)
+			{
+				Arguments.Append(" -q");
+			}
+			if((Options & SyncOptions.ReopenMovedFiles) != 0)
+			{
+				Arguments.Append(" -r");
+			}
+			if((Options & SyncOptions.Safe) != 0)
+			{
+				Arguments.Append(" -s");
+			}
+			if(MaxFiles != -1)
+			{
+				Arguments.AppendFormat(" -m {0}", MaxFiles);
+			}
+			if(NumThreads != -1)
+			{
+				Arguments.AppendFormat(" --parallel-threads={0}", NumThreads);
+				if(Batch != -1)
+				{
+					Arguments.AppendFormat(",batch={0}", Batch);
+				}
+				if(BatchSize != -1)
+				{
+					Arguments.AppendFormat(",batchsize={0}", BatchSize);
+				}
+				if(Min != -1)
+				{
+					Arguments.AppendFormat(",min={0}", Min);
+				}
+				if(MinSize != -1)
+				{
+					Arguments.AppendFormat(",minsize={0}", MinSize);
+				}
+			}
+			foreach (string FileSpec in FileSpecs)
+			{
+				Arguments.AppendFormat(" \"{0}\"", FileSpec);
+			}
+			return Arguments.ToString();
+		}
+
+		#endregion
+
+		#region p4 unshelve
+
+		/// <summary>
+		/// Restore shelved files from a pending change into a workspace
+		/// </summary>
+		/// <param name="ChangeNumber">The changelist containing shelved files</param>
+		/// <param name="IntoChangeNumber">The changelist to receive the unshelved files</param>
+		/// <param name="UsingBranchSpec">The branchspec to use when unshelving files</param>
+		/// <param name="UsingStream">Specifies the use of a stream-derived branch view to map the shelved files between the specified stream and its parent stream.</param>
+		/// <param name="ForceParentStream">Unshelve to the specified parent stream. Overrides the parent defined in the source stream specification.</param>
+		/// <param name="Options">Options for the command</param>
+		/// <param name="FileSpecs">Files to unshelve</param>
+		/// <returns>Response from the server</returns>
+		public PerforceResponseList<UnshelveRecord> Unshelve(int ChangeNumber, int IntoChangeNumber, string UsingBranchSpec, string UsingStream, string ForceParentStream, UnshelveOptions Options, params string[] FileSpecs)
+		{
+			StringBuilder Arguments = new StringBuilder("unshelve");
+			Arguments.AppendFormat(" -s {0}", ChangeNumber);
+			if((Options & UnshelveOptions.ForceOverwrite) != 0)
+			{
+				Arguments.Append(" -f");
+			}
+			if((Options & UnshelveOptions.PreviewOnly) != 0)
+			{
+				Arguments.Append(" -n");
+			}
+			if(IntoChangeNumber != -1)
+			{
+				Arguments.AppendFormat(" -c {0}", IntoChangeNumber);
+			}
+			if(UsingBranchSpec != null)
+			{
+				Arguments.AppendFormat(" -b \"{0}\"", UsingBranchSpec);
+			}
+			if(UsingStream != null)
+			{
+				Arguments.AppendFormat(" -S \"{0}\"", UsingStream);
+			}
+			if(ForceParentStream != null)
+			{
+				Arguments.AppendFormat(" -P \"{0}\"", ForceParentStream);
+			}
+			foreach(string FileSpec in FileSpecs)
+			{
+				Arguments.AppendFormat(" \"{0}\"", FileSpec);
+			}
+
+			return Command<UnshelveRecord>(Arguments.ToString(), null);
+		}
+
+		#endregion
+
+		#region p4 user
+
+		/// <summary>
+		/// Enumerates all streams in a depot
+		/// </summary>
+		/// <param name="UserName">Name of the user to fetch information for</param>
+		/// <returns>Response from the server</returns>
+		public PerforceResponse<UserRecord> User(string UserName)
+		{
+			StringBuilder Arguments = new StringBuilder("user");
+			Arguments.AppendFormat(" -o \"{0}\"", UserName);
+			return SingleResponseCommand<UserRecord>(Arguments.ToString(), null);
+		}
+
+		#endregion
+
+		#region p4 where
+
+		/// <summary>
+		/// Retrieves the location of a file of set of files in the workspace
+		/// </summary>
+		/// <param name="FileSpecs">Patterns for the files to query</param>
+		/// <returns>List of responses from the server</returns>
+		public PerforceResponseList<WhereRecord> Where(params string[] FileSpecs)
+		{
+			StringBuilder Arguments = new StringBuilder("where");
+			foreach(string FileSpec in FileSpecs)
+			{
+				Arguments.AppendFormat(" \"{0}\"", FileSpec);
+			}
+			return Command<WhereRecord>(Arguments.ToString(), null);
+		}
+
+		#endregion
+	}
+}
diff --git a/Engine/Source/Programs/DotNETCommon/DotNETUtilities/Perforce/PerforceError.cs b/Engine/Source/Programs/DotNETCommon/DotNETUtilities/Perforce/PerforceError.cs
new file mode 100644
--- /dev/null
+++ b/Engine/Source/Programs/DotNETCommon/DotNETUtilities/Perforce/PerforceError.cs
@@ -0,0 +1,79 @@
+using System;
+using System.Collections.Generic;
+using System.Linq;
+using System.Text;
+using System.Threading.Tasks;
+
+namespace Tools.DotNETCommon.Perforce
+{
+	#pragma warning disable CS1591
+
+	/// <summary>
+	/// Error severity codes. Taken from the p4java documentation.
+	/// </summary>
+	public enum PerforceSeverityCode
+	{
+		Empty = 0,
+		Info = 1,
+		Warning = 2,
+		Failed = 3,
+		Fatal = 4,
+	}
+
+	/// <summary>
+	/// Generic error codes that can be returned by the Perforce server. Taken from the p4java documentation.
+	/// </summary>
+	public enum PerforceGenericCode
+	{
+		None = 0,
+		Usage = 1,
+		Unknown = 2,
+		Context = 3,
+		Illegal = 4,
+		NotYet = 5,
+		Protect = 6,
+		Empty = 17,
+		Fault = 33,
+		Client = 34,
+		Admin = 35,
+		Config = 36,
+		Upgrade = 37,
+		Comm = 38,
+		TooBig = 39, 
+	}
+
+	#pragma warning restore CS1591
+
+	/// <summary>
+	/// Represents a error return value from Perforce.
+	/// </summary>
+	public class PerforceError
+	{
+		/// <summary>
+		/// The severity of this error
+		/// </summary>
+		[PerforceTagAttribute("severity")]
+		public PerforceSeverityCode Severity;
+
+		/// <summary>
+		/// The generic error code associated with this message
+		/// </summary>
+		[PerforceTagAttribute("generic")]
+		public PerforceGenericCode Generic;
+
+		/// <summary>
+		/// The message text
+		/// </summary>
+		[PerforceTagAttribute("data")]
+		public string Data;
+
+		/// <summary>
+		/// Formats this error for display in the debugger
+		/// </summary>
+		/// <returns>String representation of this object</returns>
+		public override string ToString()
+		{
+			return String.Format("{0}: {1} (Generic={2})", Severity, Data.TrimEnd(), Generic);
+		}
+	}
+}
diff --git a/Engine/Source/Programs/DotNETCommon/DotNETUtilities/Perforce/PerforceException.cs b/Engine/Source/Programs/DotNETCommon/DotNETUtilities/Perforce/PerforceException.cs
new file mode 100644
--- /dev/null
+++ b/Engine/Source/Programs/DotNETCommon/DotNETUtilities/Perforce/PerforceException.cs
@@ -0,0 +1,48 @@
+using System;
+using System.Collections.Generic;
+using System.Linq;
+using System.Text;
+using System.Threading.Tasks;
+
+namespace Tools.DotNETCommon.Perforce
+{
+	/// <summary>
+	/// Represents an exception specific to perforce
+	/// </summary>
+	public class PerforceException : Exception
+	{
+		/// <summary>
+		/// For errors returned by the server, contains the error record
+		/// </summary>
+		public PerforceError Error;
+
+		/// <summary>
+		/// Constructor
+		/// </summary>
+		/// <param name="Message">Message for the exception</param>
+		public PerforceException(string Message)
+			: base(Message)
+		{
+		}
+
+		/// <summary>
+		/// Constructor
+		/// </summary>
+		/// <param name="Format">Format string</param>
+		/// <param name="Args">Arguments for the formatted string</param>
+		public PerforceException(string Format, params object[] Args)
+			: base(String.Format(Format, Args))
+		{
+		}
+
+		/// <summary>
+		/// Constructor
+		/// </summary>
+		/// <param name="Error">The error from the server</param>
+		public PerforceException(PerforceError Error)
+			: base(Error.ToString())
+		{
+			this.Error = Error;
+		}
+	}
+}
diff --git a/Engine/Source/Programs/DotNETCommon/DotNETUtilities/Perforce/PerforceInfo.cs b/Engine/Source/Programs/DotNETCommon/DotNETUtilities/Perforce/PerforceInfo.cs
new file mode 100644
--- /dev/null
+++ b/Engine/Source/Programs/DotNETCommon/DotNETUtilities/Perforce/PerforceInfo.cs
@@ -0,0 +1,35 @@
+using System;
+using System.Collections.Generic;
+using System.Linq;
+using System.Text;
+using System.Threading.Tasks;
+
+namespace Tools.DotNETCommon.Perforce
+{
+	/// <summary>
+	/// Generic class for parsing "info" responses from Perforce
+	/// </summary>
+	public class PerforceInfo
+	{
+		/// <summary>
+		/// The severity level
+		/// </summary>
+		[PerforceTag("level")]
+		public int Level;
+
+		/// <summary>
+		/// Message data
+		/// </summary>
+		[PerforceTag("data")]
+		public string Data;
+
+		/// <summary>
+		/// Formats this error for display in the debugger
+		/// </summary>
+		/// <returns>String representation of this object</returns>
+		public override string ToString()
+		{
+			return Data;
+		}
+	}
+}
diff --git a/Engine/Source/Programs/DotNETCommon/DotNETUtilities/Perforce/PerforceResponse.cs b/Engine/Source/Programs/DotNETCommon/DotNETUtilities/Perforce/PerforceResponse.cs
new file mode 100644
--- /dev/null
+++ b/Engine/Source/Programs/DotNETCommon/DotNETUtilities/Perforce/PerforceResponse.cs
@@ -0,0 +1,150 @@
+using System;
+using System.Collections.Generic;
+using System.Linq;
+using System.Text;
+using System.Threading.Tasks;
+
+namespace Tools.DotNETCommon.Perforce
+{
+	/// <summary>
+	/// Base class for returning untyped response data
+	/// </summary>
+	public class PerforceResponse
+	{
+		/// <summary>
+		/// Stores the response data
+		/// </summary>
+		private object InternalData;
+
+		/// <summary>
+		/// Constructor
+		/// </summary>
+		/// <param name="Data">The response data</param>
+		public PerforceResponse(object Data)
+		{
+			this.InternalData = Data;
+		}
+
+		/// <summary>
+		/// True if the response is successful
+		/// </summary>
+		public bool Succeeded
+		{
+			get { return !(InternalData is PerforceError); }
+		}
+
+		/// <summary>
+		/// True if the response is an error
+		/// </summary>
+		public bool Failed
+		{
+			get { return InternalData is PerforceError; }
+		}
+
+		/// <summary>
+		/// Accessor for the succcessful response data. Throws an exception if the response is an error.
+		/// </summary>
+		public object Data
+		{
+			get
+			{
+				RequireSuccess();
+				return InternalData;
+			}
+		}
+
+		/// <summary>
+		/// Returns the info data.
+		/// </summary>
+		public PerforceInfo Info
+		{
+			get
+			{
+				RequireSuccess();
+				return InternalData as PerforceInfo;
+			}
+		}
+
+		/// <summary>
+		/// Returns the error data, or null if this is a succesful response.
+		/// </summary>
+		public PerforceError Error
+		{
+			get { return InternalData as PerforceError; }
+		}
+
+		/// <summary>
+		/// Throws an exception if the response is an error
+		/// </summary>
+		public void RequireSuccess()
+		{
+			PerforceError Error = InternalData as PerforceError;
+			if(Error != null)
+			{
+				throw new PerforceException(Error);
+			}
+		}
+
+		/// <summary>
+		/// Returns a string representation of this object for debugging
+		/// </summary>
+		/// <returns>String representation of the object for debugging</returns>
+		public override string ToString()
+		{
+			return InternalData.ToString();
+		}
+	}
+
+	/// <summary>
+	/// Represents a successful Perforce response of the given type, or an error. Throws a PerforceException with the error
+	/// text if the response value is attempted to be accessed and an error has occurred.
+	/// </summary>
+	/// <typeparam name="T">Type of data returned on success</typeparam>
+	public class PerforceResponse<T> : PerforceResponse
+	{
+		/// <summary>
+		/// Constructor
+		/// </summary>
+		/// <param name="Data">The successful response data</param>
+		public PerforceResponse(T Data)
+			: base(Data)
+		{
+		}
+		
+		/// <summary>
+		/// Constructor
+		/// </summary>
+		/// <param name="Info">The info data</param>
+		public PerforceResponse(PerforceInfo Info)
+			: base(Info)
+		{
+		}
+		
+		/// <summary>
+		/// Constructor
+		/// </summary>
+		/// <param name="Error">The error data</param>
+		public PerforceResponse(PerforceError Error)
+			: base(Error)
+		{
+		}
+
+		/// <summary>
+		/// Construct a typed response from an untyped response
+		/// </summary>
+		/// <param name="UntypedResponse">The untyped response</param>
+		public PerforceResponse(PerforceResponse UntypedResponse)
+			: base(UntypedResponse.Error ?? UntypedResponse.Info ?? (object)(T)UntypedResponse.Data)
+		{
+
+		}
+
+		/// <summary>
+		/// Accessor for the succcessful response data. Throws an exception if the response is an error.
+		/// </summary>
+		public new T Data
+		{
+			get { return (T)base.Data; }
+		}
+	}
+}
diff --git a/Engine/Source/Programs/DotNETCommon/DotNETUtilities/Perforce/PerforceResponseList.cs b/Engine/Source/Programs/DotNETCommon/DotNETUtilities/Perforce/PerforceResponseList.cs
new file mode 100644
--- /dev/null
+++ b/Engine/Source/Programs/DotNETCommon/DotNETUtilities/Perforce/PerforceResponseList.cs
@@ -0,0 +1,59 @@
+using System;
+using System.Collections.Generic;
+using System.Linq;
+using System.Text;
+using System.Threading.Tasks;
+
+namespace Tools.DotNETCommon.Perforce
+{
+	/// <summary>
+	/// Represents a list of responses from the Perforce server. Within the list, individual responses
+	/// may indicate success or failure.
+	/// </summary>
+	/// <typeparam name="T">Successful response type</typeparam>
+	public class PerforceResponseList<T> : List<PerforceResponse<T>> where T : class
+	{
+		/// <summary>
+		/// Whether all responses in this list are successful
+		/// </summary>
+		public bool Succeeded
+		{
+			get { return this.All(x => x.Succeeded); }
+		}
+
+		/// <summary>
+		/// Returns the first error, or null.
+		/// </summary>
+		public PerforceError FirstError
+		{
+			get { return Errors.FirstOrDefault(); }
+		}
+
+		/// <summary>
+		/// Sequence of all the data objects from the responses.
+		/// </summary>
+		public IEnumerable<T> Data
+		{
+			get { return this.Select(x => x.Data).Where(x => x != null); }
+		}
+
+		/// <summary>
+		/// Sequence of all the error responses.
+		/// </summary>
+		public IEnumerable<PerforceError> Errors
+		{
+			get { return this.Where(x => !x.Succeeded).Select(x => x.Error); }
+		}
+
+		/// <summary>
+		/// Throws an exception if any response is an error
+		/// </summary>
+		public void RequireSuccess()
+		{
+			foreach (PerforceResponse<T> Response in this)
+			{
+				Response.RequireSuccess();
+			}
+		}
+	}
+}
diff --git a/Engine/Source/Programs/DotNETCommon/DotNETUtilities/Perforce/PerforceUtils.cs b/Engine/Source/Programs/DotNETCommon/DotNETUtilities/Perforce/PerforceUtils.cs
new file mode 100644
--- /dev/null
+++ b/Engine/Source/Programs/DotNETCommon/DotNETUtilities/Perforce/PerforceUtils.cs
@@ -0,0 +1,32 @@
+using System;
+using System.Collections.Generic;
+using System.Linq;
+using System.Text;
+using System.Threading.Tasks;
+
+namespace Tools.DotNETCommon.Perforce
+{
+	class PerforceUtils
+	{
+		static public string EscapePath(string Path)
+		{
+			string NewPath = Path;
+			NewPath = NewPath.Replace("%", "%25");
+			NewPath = NewPath.Replace("*", "%2A");
+			NewPath = NewPath.Replace("#", "%23");
+			NewPath = NewPath.Replace("@", "%40");
+			return NewPath;
+		}
+
+		static public string UnescapePath(string Path)
+		{
+			string NewPath = Path;
+			NewPath = NewPath.Replace("%40", "@");
+			NewPath = NewPath.Replace("%23", "#");
+			NewPath = NewPath.Replace("%2A", "*");
+			NewPath = NewPath.Replace("%2a", "*");
+			NewPath = NewPath.Replace("%25", "%");
+			return NewPath;
+		}
+	}
+}
diff --git a/Engine/Source/Programs/DotNETCommon/DotNETUtilities/Perforce/Records/AddOptions.cs b/Engine/Source/Programs/DotNETCommon/DotNETUtilities/Perforce/Records/AddOptions.cs
new file mode 100644
--- /dev/null
+++ b/Engine/Source/Programs/DotNETCommon/DotNETUtilities/Perforce/Records/AddOptions.cs
@@ -0,0 +1,40 @@
+using System;
+using System.Collections.Generic;
+using System.Linq;
+using System.Text;
+using System.Threading.Tasks;
+
+namespace Tools.DotNETCommon.Perforce
+{
+	/// <summary>
+	/// Options for the p4 add command
+	/// </summary>
+	[Flags]
+	public enum AddOptions
+	{
+		/// <summary>
+		/// No options specified
+		/// </summary>
+		None = 0,
+
+		/// <summary>
+		/// Downgrade file open status to simple add.
+		/// </summary>
+		DowngradeToAdd = 1,
+
+		/// <summary>
+		/// Force inclusion of wildcards in filenames.
+		/// </summary>
+		IncludeWildcards = 2,
+
+		/// <summary>
+		/// Do not perform any ignore checking; ignore any settings specified by P4IGNORE.
+		/// </summary>
+		NoIgnore = 4,
+
+		/// <summary>
+		/// Preview which files would be opened for add, without actually changing any files or metadata.
+		/// </summary>
+		PreviewOnly = 8,
+	}
+}
diff --git a/Engine/Source/Programs/DotNETCommon/DotNETUtilities/Perforce/Records/AddRecord.cs b/Engine/Source/Programs/DotNETCommon/DotNETUtilities/Perforce/Records/AddRecord.cs
new file mode 100644
--- /dev/null
+++ b/Engine/Source/Programs/DotNETCommon/DotNETUtilities/Perforce/Records/AddRecord.cs
@@ -0,0 +1,53 @@
+using System;
+using System.Collections.Generic;
+using System.Linq;
+using System.Text;
+using System.Threading.Tasks;
+
+namespace Tools.DotNETCommon.Perforce
+{
+	/// <summary>
+	/// Contains information about an added file
+	/// </summary>
+	public class AddRecord
+	{
+		/// <summary>
+		/// Path to the file in the depot
+		/// </summary>
+		[PerforceTag("depotFile")]
+		public string DepotFile;
+
+		/// <summary>
+		/// Path to the file in the workspace
+		/// </summary>
+		[PerforceTag("clientFile")]
+		public string ClientFile;
+
+		/// <summary>
+		/// The working revision number of the file that was synced
+		/// </summary>
+		[PerforceTag("workRev")]
+		public int WorkingRevision;
+
+		/// <summary>
+		/// Action taken when syncing the file
+		/// </summary>
+		[PerforceTag("action")]
+		public string Action;
+
+		/// <summary>
+		/// Type of the file
+		/// </summary>
+		[PerforceTag("type")]
+		public string Type;
+
+		/// <summary>
+		/// Format this record for display in the debugger
+		/// </summary>
+		/// <returns>Summary of this revision</returns>
+		public override string ToString()
+		{
+			return DepotFile;
+		}
+	}
+}
diff --git a/Engine/Source/Programs/DotNETCommon/DotNETUtilities/Perforce/Records/ChangeRecord.cs b/Engine/Source/Programs/DotNETCommon/DotNETUtilities/Perforce/Records/ChangeRecord.cs
new file mode 100644
--- /dev/null
+++ b/Engine/Source/Programs/DotNETCommon/DotNETUtilities/Perforce/Records/ChangeRecord.cs
@@ -0,0 +1,62 @@
+using System;
+using System.Collections.Generic;
+using System.Linq;
+using System.Text;
+using System.Threading.Tasks;
+
+namespace Tools.DotNETCommon.Perforce
+{
+	/// <summary>
+	/// Contains summary information about a changelist
+	/// </summary>
+	public class ChangeRecord
+	{
+		/// <summary>
+		/// The changelist number. -1 for "new".
+		/// </summary>
+		[PerforceTag("Change")]
+		public int Number = -1;
+
+		/// <summary>
+		/// The date the change was last modified
+		/// </summary>
+		[PerforceTag("Date", Optional = true)]
+		public DateTime Date;
+
+		/// <summary>
+		/// The user that owns or submitted the change
+		/// </summary>
+		[PerforceTag("User")]
+		public string User;
+
+		/// <summary>
+		/// The client that owns the change
+		/// </summary>
+		[PerforceTag("Client")]
+		public string Client;
+
+		/// <summary>
+		/// Current changelist status
+		/// </summary>
+		[PerforceTag("Status")]
+		public ChangeStatus Status;
+
+		/// <summary>
+		/// Whether the change is restricted or not
+		/// </summary>
+		[PerforceTag("Type")]
+		public ChangeType Type;
+
+		/// <summary>
+		/// Description for the changelist
+		/// </summary>
+		[PerforceTag("Description")]
+		public string Description;
+
+		/// <summary>
+		/// Files that are open in this changelist
+		/// </summary>
+		[PerforceTag("Files", Optional = true)]
+		public List<string> Files = new List<string>();
+	}
+}
diff --git a/Engine/Source/Programs/DotNETCommon/DotNETUtilities/Perforce/Records/ChangeStatus.cs b/Engine/Source/Programs/DotNETCommon/DotNETUtilities/Perforce/Records/ChangeStatus.cs
new file mode 100644
--- /dev/null
+++ b/Engine/Source/Programs/DotNETCommon/DotNETUtilities/Perforce/Records/ChangeStatus.cs
@@ -0,0 +1,37 @@
+using System;
+using System.Collections.Generic;
+using System.Linq;
+using System.Text;
+using System.Threading.Tasks;
+
+namespace Tools.DotNETCommon.Perforce
+{
+	/// <summary>
+	/// Indicates the status of a changelist
+	/// </summary>
+	public enum ChangeStatus
+	{
+		/// <summary>
+		/// Include all changes
+		/// </summary>
+		All,
+
+		/// <summary>
+		/// The change is still pending
+		/// </summary>
+		[PerforceEnum("pending")]
+		Pending,
+
+		/// <summary>
+		/// The change has been submitted
+		/// </summary>
+		[PerforceEnum("submitted")]
+		Submitted,
+
+		/// <summary>
+		/// The change is shelved
+		/// </summary>
+		[PerforceEnum("shelved")]
+		Shelved,
+	}
+}
diff --git a/Engine/Source/Programs/DotNETCommon/DotNETUtilities/Perforce/Records/ChangeType.cs b/Engine/Source/Programs/DotNETCommon/DotNETUtilities/Perforce/Records/ChangeType.cs
new file mode 100644
--- /dev/null
+++ b/Engine/Source/Programs/DotNETCommon/DotNETUtilities/Perforce/Records/ChangeType.cs
@@ -0,0 +1,31 @@
+using System;
+using System.Collections.Generic;
+using System.Linq;
+using System.Text;
+using System.Threading.Tasks;
+
+namespace Tools.DotNETCommon.Perforce
+{
+	/// <summary>
+	/// The type of a changelist
+	/// </summary>
+	public enum ChangeType
+	{
+		/// <summary>
+		/// When creating a new changelist, leaves the changelist type unspecified.
+		/// </summary>
+		Unspecified,
+
+		/// <summary>
+		/// This change is visible to anyone
+		/// </summary>
+		[PerforceEnum("public")]
+		Public,
+
+		/// <summary>
+		/// This change is restricted
+		/// </summary>
+		[PerforceEnum("restricted")]
+		Restricted,
+	}
+}
diff --git a/Engine/Source/Programs/DotNETCommon/DotNETUtilities/Perforce/Records/ChangesOptions.cs b/Engine/Source/Programs/DotNETCommon/DotNETUtilities/Perforce/Records/ChangesOptions.cs
new file mode 100644
--- /dev/null
+++ b/Engine/Source/Programs/DotNETCommon/DotNETUtilities/Perforce/Records/ChangesOptions.cs
@@ -0,0 +1,45 @@
+using System;
+using System.Collections.Generic;
+using System.Linq;
+using System.Text;
+using System.Threading.Tasks;
+
+namespace Tools.DotNETCommon.Perforce
+{
+	/// <summary>
+	/// Options for the 'changes' command
+	/// </summary>
+	[Flags]
+	public enum ChangesOptions
+	{
+		/// <summary>
+		/// No options specified
+		/// </summary>
+		None = 0,
+
+		/// <summary>
+		/// Include changelists that affected files that were integrated with the specified files.
+		/// </summary>
+		IncludeIntegrations = 1,
+
+		/// <summary>
+		/// Display the time as well as the date of each change.
+		/// </summary>
+		IncludeTimes = 2,
+
+		/// <summary>
+		/// List long output, with the full text of each changelist description.
+		/// </summary>
+		LongOutput = 4,
+
+		/// <summary>
+		/// List long output, with the full text of each changelist description truncated at 250 characters.
+		/// </summary>
+		TruncatedLongOutput = 8,
+
+		/// <summary>
+		/// View restricted changes (requires admin permission)
+		/// </summary>
+		IncludeRestricted = 16,
+	}
+}
diff --git a/Engine/Source/Programs/DotNETCommon/DotNETUtilities/Perforce/Records/ChangesRecord.cs b/Engine/Source/Programs/DotNETCommon/DotNETUtilities/Perforce/Records/ChangesRecord.cs
new file mode 100644
--- /dev/null
+++ b/Engine/Source/Programs/DotNETCommon/DotNETUtilities/Perforce/Records/ChangesRecord.cs
@@ -0,0 +1,62 @@
+using System;
+using System.Collections.Generic;
+using System.Linq;
+using System.Text;
+using System.Threading.Tasks;
+
+namespace Tools.DotNETCommon.Perforce
+{
+	/// <summary>
+	/// Contains summary information about a changelist
+	/// </summary>
+	public class ChangesRecord
+	{
+		/// <summary>
+		/// The changelist number
+		/// </summary>
+		[PerforceTag("change")]
+		public int Number;
+
+		/// <summary>
+		/// The date the change was last modified
+		/// </summary>
+		[PerforceTag("time")]
+		public DateTime Time;
+
+		/// <summary>
+		/// The user that owns or submitted the change
+		/// </summary>
+		[PerforceTag("user")]
+		public string User;
+
+		/// <summary>
+		/// The client that owns the change
+		/// </summary>
+		[PerforceTag("client")]
+		public string Client;
+
+		/// <summary>
+		/// Current changelist status
+		/// </summary>
+		[PerforceTag("status")]
+		public ChangeStatus Status;
+
+		/// <summary>
+		/// Whether the change is restricted or not
+		/// </summary>
+		[PerforceTag("changeType")]
+		public ChangeType Type;
+
+		/// <summary>
+		/// The path affected by this change.
+		/// </summary>
+		[PerforceTag("path")]
+		public string Path;
+
+		/// <summary>
+		/// Description for the changelist
+		/// </summary>
+		[PerforceTag("desc")]
+		public string Description;
+	}
+}
diff --git a/Engine/Source/Programs/DotNETCommon/DotNETUtilities/Perforce/Records/ClientLineEndings.cs b/Engine/Source/Programs/DotNETCommon/DotNETUtilities/Perforce/Records/ClientLineEndings.cs
new file mode 100644
--- /dev/null
+++ b/Engine/Source/Programs/DotNETCommon/DotNETUtilities/Perforce/Records/ClientLineEndings.cs
@@ -0,0 +1,49 @@
+using System;
+using System.Collections.Generic;
+using System.Linq;
+using System.Text;
+using System.Threading.Tasks;
+
+namespace Tools.DotNETCommon.Perforce
+{
+	/// <summary>
+	/// Specifies how to treat line endings on files synced to the client
+	/// </summary>
+	public enum ClientLineEndings
+	{
+		/// <summary>
+		/// This value is not specified
+		/// </summary>
+		Unspecified,
+
+		/// <summary>
+		/// Use mode native to the client (default)
+		/// </summary>
+		[PerforceEnum("local")]
+		Local,
+
+		/// <summary>
+		/// UNIX-style (and Mac OS X) line endings: LF
+		/// </summary>
+		[PerforceEnum("unix")]
+		Unix,
+
+		/// <summary>
+		/// Mac pre-OS X: CR only
+		/// </summary>
+		[PerforceEnum("mac")]
+		Mac,
+
+		/// <summary>
+		/// Windows-style: CR + LF.
+		/// </summary>
+		[PerforceEnum("win")]
+		Windows,
+
+		/// <summary>
+		/// The share option normalizes mixed line-endings into UNIX line-end format.
+		/// </summary>
+		[PerforceEnum("share")]
+		Share,
+	}
+}
diff --git a/Engine/Source/Programs/DotNETCommon/DotNETUtilities/Perforce/Records/ClientOptions.cs b/Engine/Source/Programs/DotNETCommon/DotNETUtilities/Perforce/Records/ClientOptions.cs
new file mode 100644
--- /dev/null
+++ b/Engine/Source/Programs/DotNETCommon/DotNETUtilities/Perforce/Records/ClientOptions.cs
@@ -0,0 +1,92 @@
+using System;
+using System.Collections.Generic;
+using System.Linq;
+using System.Text;
+using System.Threading.Tasks;
+
+namespace Tools.DotNETCommon.Perforce
+{
+	/// <summary>
+	/// Options for the client
+	/// </summary>
+	[Flags]
+	public enum ClientOptions
+	{
+		/// <summary>
+		/// No options are set
+		/// </summary>
+		None = 0,
+
+		/// <summary>
+		/// If set, unopened files in the workspace are left writable.
+		/// </summary>
+		[PerforceEnum("allwrite")]
+		AllWrite = 1,
+
+		/// <summary>
+		/// Opposite of AllWrite.
+		/// </summary>
+		[PerforceEnum("noallwrite")]
+		NoAllWrite = 2,
+
+		/// <summary>
+		/// If set, a p4 sync overwrites ("clobbers") writable-but-unopened files in the workspace that have the same name as the newly-synced files.
+		/// </summary>
+		[PerforceEnum("clobber")]
+		Clobber = 4,
+
+		/// <summary>
+		/// Opposite of Clobber.
+		/// </summary>
+		[PerforceEnum("noclobber")]
+		NoClobber = 8,
+
+		/// <summary>
+		/// If set, a p4 sync overwrites ("clobbers") writable-but-unopened files in the workspace that have the same name as the newly-synced files.
+		/// </summary>
+		[PerforceEnum("compress")]
+		Compress = 16,
+
+		/// <summary>
+		/// Opposite of Compress.
+		/// </summary>
+		[PerforceEnum("nocompress")]
+		NoCompress = 32,
+
+		/// <summary>
+		/// If set, a p4 sync overwrites ("clobbers") writable-but-unopened files in the workspace that have the same name as the newly-synced files.
+		/// </summary>
+		[PerforceEnum("locked")]
+		Locked = 64,
+
+		/// <summary>
+		/// Opposite of Locked.
+		/// </summary>
+		[PerforceEnum("unlocked")]
+		Unlocked = 128,
+
+		/// <summary>
+		/// For files without the +m (modtime) file type modifier, the modification date (on the local filesystem) of a newly synced file is the datestamp on the file when the file was last modified.
+		/// </summary>
+		[PerforceEnum("modtime")]
+		ModTime = 256,
+
+		/// <summary>
+		/// For files without the +m (modtime) file type modifier, the modification date is the date and time of sync, regardless of version.
+		/// </summary>
+		[PerforceEnum("nomodtime")]
+		NoModTime = 512,
+
+		/// <summary>
+		/// If set, p4 sync deletes empty directories in a workspace if all files in the directory have been removed.
+		/// </summary>
+		[PerforceEnum("rmdir")]
+		Rmdir = 1024,
+
+		/// <summary>
+		/// Opposite of Rmdir.
+		/// </summary>
+		[PerforceEnum("normdir")]
+		NoRmdir = 2048,
+	}
+}
diff --git a/Engine/Source/Programs/DotNETCommon/DotNETUtilities/Perforce/Records/ClientRecord.cs b/Engine/Source/Programs/DotNETCommon/DotNETUtilities/Perforce/Records/ClientRecord.cs
new file mode 100644
--- /dev/null
+++ b/Engine/Source/Programs/DotNETCommon/DotNETUtilities/Perforce/Records/ClientRecord.cs
@@ -0,0 +1,92 @@
+using System;
+using System.Collections.Generic;
+using System.Linq;
+using System.Text;
+using System.Threading.Tasks;
+
+namespace Tools.DotNETCommon.Perforce
+{
+	/// <summary>
+	/// Represents a Perforce clientspec
+	/// </summary>
+	public class ClientRecord
+	{
+		/// <summary>
+		/// The client workspace name, as specified in the P4CLIENT environment variable or its equivalents.
+		/// </summary>
+		[PerforceTag("Client")]
+		public string Name;
+
+		/// <summary>
+		/// The Perforce user name of the user who owns the workspace.
+		/// </summary>
+		[PerforceTag("Owner")]
+		public string Owner;
+
+		/// <summary>
+		/// The date the workspace specification was last modified.
+		/// </summary>
+		[PerforceTag("Update", Optional = true)]
+		public DateTime Update;
+
+		/// <summary>
+		/// The date and time that the workspace was last used in any way.
+		/// </summary>
+		[PerforceTag("Access", Optional = true)]
+		public DateTime Access;
+
+		/// <summary>
+		/// The name of the workstation on which this workspace resides.
+		/// </summary>
+		[PerforceTag("Host", Optional = true)]
+		public string Host;
+
+		/// <summary>
+		/// A textual description of the workspace.
+		/// </summary>
+		[PerforceTag("Description", Optional = true)]
+		public string Description;
+
+		/// <summary>
+		/// The directory (on the local host) relative to which all the files in the View: are specified. 
+		/// </summary>
+		[PerforceTag("Root")]
+		public string Root;
+
+		/// <summary>
+		/// A set of seven switches that control particular workspace options.
+		/// </summary>
+		[PerforceTag("Options")]
+		public ClientOptions Options;
+
+		/// <summary>
+		/// Options to govern the default behavior of p4 submit.
+		/// </summary>
+		[PerforceTag("SubmitOptions")]
+		public ClientSubmitOptions SubmitOptions;
+
+		/// <summary>
+		/// Configure carriage-return/linefeed (CR/LF) conversion. 
+		/// </summary>
+		[PerforceTag("LineEnd")]
+		public ClientLineEndings LineEnd;
+
+		/// <summary>
+		/// Associates the workspace with the specified stream.
+		/// </summary>
+		[PerforceTag("Stream", Optional = true)]
+		public string Stream;
+
+		/// <summary>
+		/// The type of client.
+		/// </summary>
+		[PerforceTag("Type", Optional = true)]
+		public string Type;
+
+		/// <summary>
+		/// Specifies the mappings between files in the depot and files in the workspace.
+		/// </summary>
+		[PerforceTag("View")]
+		public List<string> View = new List<string>();
+	}
+}
diff --git a/Engine/Source/Programs/DotNETCommon/DotNETUtilities/Perforce/Records/ClientSubmitOptions.cs b/Engine/Source/Programs/DotNETCommon/DotNETUtilities/Perforce/Records/ClientSubmitOptions.cs
new file mode 100644
--- /dev/null
+++ b/Engine/Source/Programs/DotNETCommon/DotNETUtilities/Perforce/Records/ClientSubmitOptions.cs
@@ -0,0 +1,55 @@
+using System;
+using System.Collections.Generic;
+using System.Linq;
+using System.Text;
+using System.Threading.Tasks;
+
+namespace Tools.DotNETCommon.Perforce
+{
+	/// <summary>
+	/// Client options for controlling the default behavior of p4 submit
+	/// </summary>
+	public enum ClientSubmitOptions
+	{
+		/// <summary>
+		/// This field is not set
+		/// </summary>
+		Unspecified,
+
+		/// <summary>
+		/// All open files (with or without changes) are submitted to the depot.
+		/// </summary>
+		[PerforceEnum("submitunchanged")]
+		SubmitUnchanged,
+
+		/// <summary>
+		/// All open files (with or without changes) are submitted to the depot, and all files are automatically reopened in the default changelist.
+		/// </summary>
+		[PerforceEnum("submitunchanged+reopen")]
+		SubmitUnchangedAndReopen,
+
+		/// <summary>
+		/// Only those files with content, type, or resolved changes are submitted to the depot.
+		/// </summary>
+		[PerforceEnum("revertunchanged")]
+		RevertUnchanged,
+
+		/// <summary>
+		/// Only those files with content, type, or resolved changes are submitted to the depot and reopened in the default changelist.
+		/// </summary>
+		[PerforceEnum("revertunchanged+reopen")]
+		RevertUnchangedAndReopen,
+
+		/// <summary>
+		/// Only those files with content, type, or resolved changes are submitted to the depot. Any unchanged files are moved to the default changelist.
+		/// </summary>
+		[PerforceEnum("leaveunchanged")]
+		LeaveUnchanged,
+
+		/// <summary>
+		/// Only those files with content, type, or resolved changes are submitted to the depot. 
+		/// </summary>
+		[PerforceEnum("leaveunchanged+reopen")]
+		LeaveUnchangedAndReopen,
+	}
+}
diff --git a/Engine/Source/Programs/DotNETCommon/DotNETUtilities/Perforce/Records/ClientsOptions.cs b/Engine/Source/Programs/DotNETCommon/DotNETUtilities/Perforce/Records/ClientsOptions.cs
new file mode 100644
--- /dev/null
+++ b/Engine/Source/Programs/DotNETCommon/DotNETUtilities/Perforce/Records/ClientsOptions.cs
@@ -0,0 +1,40 @@
+using System;
+using System.Collections.Generic;
+using System.Linq;
+using System.Text;
+using System.Threading.Tasks;
+
+namespace Tools.DotNETCommon.Perforce
+{
+	/// <summary>
+	/// Options for the p4 clients command
+	/// </summary>
+	[Flags]
+	public enum ClientsOptions
+	{
+		/// <summary>
+		/// No options
+		/// </summary>
+		None = 0,
+
+		/// <summary>
+		/// List all client workspaces, not just workspaces bound to this server
+		/// </summary>
+		All = 1,
+
+		/// <summary>
+		/// Treat the filter argument as case sensitive
+		/// </summary>
+		CaseSensitiveFilter = 2,
+
+		/// <summary>
+		/// Display the time as well as the date of the last update to the workspace.
+		/// </summary>
+		WithTimes = 4,
+
+		/// <summary>
+		/// List only client workspaces unloaded with p4 unload.
+		/// </summary>
+		Unloaded = 8,
+	}
+}
diff --git a/Engine/Source/Programs/DotNETCommon/DotNETUtilities/Perforce/Records/ClientsRecord.cs b/Engine/Source/Programs/DotNETCommon/DotNETUtilities/Perforce/Records/ClientsRecord.cs
new file mode 100644
--- /dev/null
+++ b/Engine/Source/Programs/DotNETCommon/DotNETUtilities/Perforce/Records/ClientsRecord.cs
@@ -0,0 +1,89 @@
+using System;
+using System.Collections.Generic;
+using System.Linq;
+using System.Text;
+using System.Threading.Tasks;
+
+namespace Tools.DotNETCommon.Perforce
+{
+	/// <summary>
+	/// Record returned by the p4 clients command
+	/// </summary>
+	public class ClientsRecord
+	{
+		/// <summary>
+		/// The client workspace name, as specified in the P4CLIENT environment variable or its equivalents.
+		/// </summary>
+		[PerforceTag("client")]
+		public string Name;
+
+		/// <summary>
+		/// The Perforce user name of the user who owns the workspace.
+		/// </summary>
+		[PerforceTag("Owner")]
+		public string Owner;
+
+		/// <summary>
+		/// The date the workspace specification was last modified.
+		/// </summary>
+		[PerforceTag("Update")]
+		public DateTime Update;
+
+		/// <summary>
+		/// The date and time that the workspace was last used in any way.
+		/// </summary>
+		[PerforceTag("Access")]
+		public DateTime Access;
+
+		/// <summary>
+		/// The name of the workstation on which this workspace resides.
+		/// </summary>
+		[PerforceTag("Host", Optional = true)]
+		public string Host;
+
+		/// <summary>
+		/// A textual description of the workspace.
+		/// </summary>
+		[PerforceTag("Description", Optional = true)]
+		public string Description;
+
+		/// <summary>
+		/// The directory (on the local host) relative to which all the files in the View: are specified. 
+		/// </summary>
+		[PerforceTag("Root")]
+		public string Root;
+
+		/// <summary>
+		/// A set of seven switches that control particular workspace options.
+		/// </summary>
+		[PerforceTag("Options")]
+		public ClientOptions Options;
+
+		/// <summary>
+		/// Options to govern the default behavior of p4 submit.
+		/// </summary>
+		[PerforceTag("SubmitOptions")]
+		public ClientSubmitOptions SubmitOptions;
+
+		/// <summary>
+		/// Configure carriage-return/linefeed (CR/LF) conversion. 
+		/// </summary>
+		[PerforceTag("LineEnd")]
+		public ClientLineEndings LineEnd;
+
+		/// <summary>
+		/// Associates the workspace with the specified stream.
+		/// </summary>
+		[PerforceTag("Stream", Optional = true)]
+		public string Stream;
+
+		/// <summary>
+		/// Format this record for display in the debugger
+		/// </summary>
+		/// <returns>Summary of this revision</returns>
+		public override string ToString()
+		{
+			return Name;
+		}
+	}
+}
diff --git a/Engine/Source/Programs/DotNETCommon/DotNETUtilities/Perforce/Records/DeleteChangeOptions.cs b/Engine/Source/Programs/DotNETCommon/DotNETUtilities/Perforce/Records/DeleteChangeOptions.cs
new file mode 100644
--- /dev/null
+++ b/Engine/Source/Programs/DotNETCommon/DotNETUtilities/Perforce/Records/DeleteChangeOptions.cs
@@ -0,0 +1,29 @@
+using System;
+using System.Collections.Generic;
+using System.Linq;
+using System.Text;
+using System.Threading.Tasks;
+
+namespace Tools.DotNETCommon.Perforce
+{
+	/// <summary>
+	/// Options for the DeleteChange command (p4 change -d)
+	/// </summary>
+	public enum DeleteChangeOptions
+	{
+		/// <summary>
+		/// No options are specified
+		/// </summary>
+		None = 0,
+
+		/// <summary>
+		/// Forcibly delete a previously submitted changelist.
+		/// </summary>
+		Submitted = 1,
+
+		/// <summary>
+		/// If a changelist was renumbered on submit, and you know only the original changelist number, use OriginalChangeNumber and the original changelist number to view or edit the changelist.
+		/// </summary>
+		BeforeRenumber = 2,
+	}
+}
diff --git a/Engine/Source/Programs/DotNETCommon/DotNETUtilities/Perforce/Records/DeleteClientOptions.cs b/Engine/Source/Programs/DotNETCommon/DotNETUtilities/Perforce/Records/DeleteClientOptions.cs
new file mode 100644
--- /dev/null
+++ b/Engine/Source/Programs/DotNETCommon/DotNETUtilities/Perforce/Records/DeleteClientOptions.cs
@@ -0,0 +1,30 @@
+using System;
+using System.Collections.Generic;
+using System.Linq;
+using System.Text;
+using System.Threading.Tasks;
+
+namespace Tools.DotNETCommon.Perforce
+{
+	/// <summary>
+	/// Options for the p4 client -d command
+	/// </summary>
+	[Flags]
+	public enum DeleteClientOptions
+	{
+		/// <summary>
+		/// No options
+		/// </summary>
+		None = 0,
+
+		/// <summary>
+		/// Administrators can use the -f option to delete or modify locked workspaces owned by other users.
+		/// </summary>
+		Force = 1,
+
+		/// <summary>
+		/// Allows the deletion of a client even when that client contains shelved changes. Requires the Force option.
+		/// </summary>
+		DeleteShelved = 2,
+	}
+}
diff --git a/Engine/Source/Programs/DotNETCommon/DotNETUtilities/Perforce/Records/DescribeFileRecord.cs b/Engine/Source/Programs/DotNETCommon/DotNETUtilities/Perforce/Records/DescribeFileRecord.cs
new file mode 100644
--- /dev/null
+++ b/Engine/Source/Programs/DotNETCommon/DotNETUtilities/Perforce/Records/DescribeFileRecord.cs
@@ -0,0 +1,59 @@
+using System;
+using System.Collections.Generic;
+using System.Linq;
+using System.Text;
+using System.Threading.Tasks;
+
+namespace Tools.DotNETCommon.Perforce
+{
+	/// <summary>
+	/// Contains information about an individual file in a returned DescribeRecord
+	/// </summary>
+	public class DescribeFileRecord
+	{
+		/// <summary>
+		/// Path to the modified file in depot syntax
+		/// </summary>
+		[PerforceTag("depotFile")]
+		public string DepotFile;
+
+		/// <summary>
+		/// The action performed on this file
+		/// </summary>
+		[PerforceTag("action")]
+		public FileAction Action;
+
+		/// <summary>
+		/// The file type after this change
+		/// </summary>
+		[PerforceTag("type")]
+		public string Type;
+
+		/// <summary>
+		/// The revision number for this file
+		/// </summary>
+		[PerforceTag("rev")]
+		public int Revision;
+
+		/// <summary>
+		/// Size of the file, or -1 if not specified
+		/// </summary>
+		[PerforceTag("fileSize", Optional = true)]
+		public readonly int FileSize;
+
+		/// <summary>
+		/// Digest of the file, or null if not specified
+		/// </summary>
+		[PerforceTag("digest", Optional = true)]
+		public readonly string Digest;
+
+		/// <summary>
+		/// Format this record for display in the debugger
+		/// </summary>
+		/// <returns>Summary of this revision</returns>
+		public override string ToString()
+		{
+			return DepotFile;
+		}
+	}
+}
diff --git a/Engine/Source/Programs/DotNETCommon/DotNETUtilities/Perforce/Records/DescribeRecord.cs b/Engine/Source/Programs/DotNETCommon/DotNETUtilities/Perforce/Records/DescribeRecord.cs
new file mode 100644
--- /dev/null
+++ b/Engine/Source/Programs/DotNETCommon/DotNETUtilities/Perforce/Records/DescribeRecord.cs
@@ -0,0 +1,77 @@
+using System;
+using System.Collections.Generic;
+using System.Linq;
+using System.Text;
+using System.Threading.Tasks;
+
+namespace Tools.DotNETCommon.Perforce
+{
+	/// <summary>
+	/// Detailed description of an individual changelist
+	/// </summary>
+	public class DescribeRecord
+	{
+		/// <summary>
+		/// The changelist number
+		/// </summary>
+		[PerforceTag("change")]
+		public int Number;
+
+		/// <summary>
+		/// The user that owns or submitted the change
+		/// </summary>
+		[PerforceTag("user")]
+		public string User;
+
+		/// <summary>
+		/// The workspace that contains or submitted the change
+		/// </summary>
+		[PerforceTag("client")]
+		public string Client;
+
+		/// <summary>
+		/// Time at which the change was submitted
+		/// </summary>
+		[PerforceTag("time")]
+		public DateTime Time;
+
+		/// <summary>
+		/// The changelist description
+		/// </summary>
+		[PerforceTag("desc")]
+		public string Description;
+
+		/// <summary>
+		/// The change status (submitted, pending, etc...)
+		/// </summary>
+		[PerforceTag("status")]
+		public ChangeStatus Status;
+
+		/// <summary>
+		/// Whether the change is restricted or not
+		/// </summary>
+		[PerforceTag("changeType")]
+		public ChangeType Type;
+
+		/// <summary>
+		/// Narrowest path which contains all files affected by this change
+		/// </summary>
+		[PerforceTag("path")]
+		public string Path;
+
+		/// <summary>
+		/// The files affected by this change
+		/// </summary>
+		[PerforceRecordList]
+		public List<DescribeFileRecord> Files = new List<DescribeFileRecord>();
+
+		/// <summary>
+		/// Format this record for display in the debugger
+		/// </summary>
+		/// <returns>Summary of this revision</returns>
+		public override string ToString()
+		{
+			return String.Format("{0}: {1}", Number, Description);
+		}
+	}
+}
diff --git a/Engine/Source/Programs/DotNETCommon/DotNETUtilities/Perforce/Records/EditOptions.cs b/Engine/Source/Programs/DotNETCommon/DotNETUtilities/Perforce/Records/EditOptions.cs
new file mode 100644
--- /dev/null
+++ b/Engine/Source/Programs/DotNETCommon/DotNETUtilities/Perforce/Records/EditOptions.cs
@@ -0,0 +1,30 @@
+using System;
+using System.Collections.Generic;
+using System.Linq;
+using System.Text;
+using System.Threading.Tasks;
+
+namespace Tools.DotNETCommon.Perforce
+{
+	/// <summary>
+	/// Options for the p4 edit command
+	/// </summary>
+	[Flags]
+	public enum EditOptions
+	{
+		/// <summary>
+		/// No options specified
+		/// </summary>
+		None = 0,
+
+		/// <summary>
+		/// Keep existing workspace files; mark the file as open for edit even if the file is not in the client view. 
+		/// </summary>
+		KeepWorkspaceFiles = 1,
+
+		/// <summary>
+		/// Preview which files would be opened for edit, without actually changing any files or metadata.
+		/// </summary>
+		PreviewOnly = 2,
+	}
+}
diff --git a/Engine/Source/Programs/DotNETCommon/DotNETUtilities/Perforce/Records/EditRecord.cs b/Engine/Source/Programs/DotNETCommon/DotNETUtilities/Perforce/Records/EditRecord.cs
new file mode 100644
--- /dev/null
+++ b/Engine/Source/Programs/DotNETCommon/DotNETUtilities/Perforce/Records/EditRecord.cs
@@ -0,0 +1,53 @@
+using System;
+using System.Collections.Generic;
+using System.Linq;
+using System.Text;
+using System.Threading.Tasks;
+
+namespace Tools.DotNETCommon.Perforce
+{
+	/// <summary>
+	/// Information about a file opened for edit
+	/// </summary>
+	public class EditRecord
+	{
+		/// <summary>
+		/// Path to the file in the depot
+		/// </summary>
+		[PerforceTag("depotFile")]
+		public string DepotFile;
+
+		/// <summary>
+		/// Path to the file in the workspace
+		/// </summary>
+		[PerforceTag("clientFile")]
+		public string ClientFile;
+
+		/// <summary>
+		/// The working revision number of the file that was synced
+		/// </summary>
+		[PerforceTag("workRev")]
+		public int WorkingRevision;
+
+		/// <summary>
+		/// Action taken when syncing the file
+		/// </summary>
+		[PerforceTag("action")]
+		public string Action;
+
+		/// <summary>
+		/// Type of the file
+		/// </summary>
+		[PerforceTag("type")]
+		public string Type;
+
+		/// <summary>
+		/// Format this record for display in the debugger
+		/// </summary>
+		/// <returns>Summary of this revision</returns>
+		public override string ToString()
+		{
+			return DepotFile;
+		}
+	}
+}
diff --git a/Engine/Source/Programs/DotNETCommon/DotNETUtilities/Perforce/Records/FStatOptions.cs b/Engine/Source/Programs/DotNETCommon/DotNETUtilities/Perforce/Records/FStatOptions.cs
new file mode 100644
--- /dev/null
+++ b/Engine/Source/Programs/DotNETCommon/DotNETUtilities/Perforce/Records/FStatOptions.cs
@@ -0,0 +1,120 @@
+using System;
+using System.Collections.Generic;
+using System.Linq;
+using System.Text;
+using System.Threading.Tasks;
+
+namespace Tools.DotNETCommon.Perforce
+{
+	/// <summary>
+	/// Options for the Fstat command
+	/// </summary>
+	[Flags]
+	public enum FStatOptions
+	{
+		/// <summary>
+		/// No options specified
+		/// </summary>
+		None = 0x0,
+
+		/// <summary>
+		/// For scripting purposes, report file information on a list of valid file arguments in full depot syntax with a valid revision number.
+		/// </summary>
+		ReportDepotSyntax = 0x01,
+
+		/// <summary>
+		/// Output all revisions for the given files, suppressing the other[...] and resolve[...] fields.
+		/// </summary>
+		AllRevisions = 0x02, 
+
+		/// <summary>
+		/// Output a fileSize field displaying the length of the file and a digest field for each revision.
+		/// </summary>
+		IncludeFileSizes = 0x04,
+
+		/// <summary>
+		/// Display the clientFile in Perforce syntax, as opposed to local syntax.
+		/// </summary>
+		ClientFileInPerforceSyntax = 0x08,
+
+		/// <summary>
+		/// Display pending integration record data for files open in the current workspace.
+		/// </summary>
+		ShowPendingIntegrations = 0x10,
+
+		/// <summary>
+		/// Shorten output by excluding client workspace data (for instance, the clientFile field).
+		/// </summary>
+		ShortenOutput = 0x20,
+
+		/// <summary>
+		/// Sort the output in reverse order.
+		/// </summary>
+		ReverseOrder = 0x40,
+
+		/// <summary>
+		/// Limit output to files mapped into the current workspace.
+		/// </summary>
+		OnlyMapped = 0x80,
+
+		/// <summary>
+		/// Limit output to files on your have list; that is, to files synced to the current workspace.
+		/// </summary>
+		OnlyHave = 0x100,
+
+		/// <summary>
+		/// Limit output to files opened at revisions not at the head revision.
+		/// </summary>
+		OnlyOpenedBeforeHead = 0x200,
+
+		/// <summary>
+		/// Limit output to open files in the current workspace.
+		/// </summary>
+		OnlyOpenInWorkspace = 0x400,
+
+		/// <summary>
+		/// Limit output to open files that have been resolved.
+		/// </summary>
+		OnlyOpenAndResolved = 0x800,
+
+		/// <summary>
+		/// Limit output to shelved files.
+		/// </summary>
+		OnlyShelved = 0x1000,
+
+		/// <summary>
+		/// Limit output to open files that are unresolved.
+		/// </summary>
+		OnlyUnresolved = 0x2000,
+
+		/// <summary>
+		/// Sort by date.
+		/// </summary>
+		SortByDate = 0x4000,
+
+		/// <summary>
+		/// Sort by have revision.
+		/// </summary>
+		SortByHaveRevision = 0x8000,
+
+		/// <summary>
+		/// Sort by head revision.
+		/// </summary>
+		SortByHeadRevision = 0x10000,
+
+		/// <summary>
+		/// Sort by file size.
+		/// </summary>
+		SortByFileSize = 0x20000,
+
+		/// <summary>
+		/// Sort by file type.
+		/// </summary>
+		SortByFileType = 0x40000,
+
+		/// <summary>
+		/// Include files in the unload depot when displaying data.
+		/// </summary>
+		IncludeFilesInUnloadDepot = 0x80000,
+	}
+}
diff --git a/Engine/Source/Programs/DotNETCommon/DotNETUtilities/Perforce/Records/FStatRecord.cs b/Engine/Source/Programs/DotNETCommon/DotNETUtilities/Perforce/Records/FStatRecord.cs
new file mode 100644
--- /dev/null
+++ b/Engine/Source/Programs/DotNETCommon/DotNETUtilities/Perforce/Records/FStatRecord.cs
@@ -0,0 +1,143 @@
+using System;
+using System.Collections.Generic;
+using System.Linq;
+using System.Text;
+using System.Threading.Tasks;
+
+namespace Tools.DotNETCommon.Perforce
+{
+	/// <summary>
+	/// Record returned by the fstat command
+	/// </summary>
+	public class FStatRecord
+	{
+		/// <summary>
+		/// Depot path to file
+		/// </summary>
+		[PerforceTag("depotFile", Optional = true)]
+		public string DepotFile;
+
+		/// <summary>
+		/// Local path to file (in local syntax by default, or in Perforce syntax with the FStatOptions.ClientFileInPerforceSyntax option)
+		/// </summary>
+		[PerforceTag("clientFile", Optional = true)]
+		public string ClientFile;
+
+		/// <summary>
+		/// Local path to file
+		/// </summary>
+		[PerforceTag("path", Optional = true)]
+		public string Path;
+
+		/// <summary>
+		/// Set if the file is open for add and it is mapped to current client workspace
+		/// </summary>
+		[PerforceTag("isMapped", Optional = true)]
+		public bool IsMapped;
+
+		/// <summary>
+		/// Set if file is shelved
+		/// </summary>
+		[PerforceTag("shelved", Optional = true)]
+		public bool Shelved;
+
+		/// <summary>
+		/// Action taken at head revision, if in depot
+		/// </summary>
+		[PerforceTag("headAction", Optional = true)]
+		public FileAction HeadAction;
+
+		/// <summary>
+		/// Head revision type, if in depot
+		/// </summary>
+		[PerforceTag("headType", Optional = true)]
+		public string HeadType;
+
+		/// <summary>
+		/// Head revision changelist time, if in depot.
+		/// </summary>
+		[PerforceTag("headTime", Optional = true)]
+		public DateTime HeadTime;
+
+		/// <summary>
+		/// Head revision number, if in depot
+		/// </summary>
+		[PerforceTag("headRev", Optional = true)]
+		public int HeadRevision;
+
+		/// <summary>
+		/// Head revision changelist number, if in depot
+		/// </summary>
+		[PerforceTag("headChange", Optional = true)]
+		public int HeadChange;
+
+		/// <summary>
+		/// Head revision modification time (the time that the file was last modified on the client before submit), if in depot.
+		/// </summary>
+		[PerforceTag("headModTime", Optional = true)]
+		public DateTime HeadModTime;
+
+		/// <summary>
+		/// Head revision of moved file
+		/// </summary>
+		[PerforceTag("movedRev", Optional = true)]
+		public int MovedRevision;
+
+		/// <summary>
+		/// Revision last synced to workspace, if on workspace
+		/// </summary>
+		[PerforceTag("haveRev", Optional = true)]
+		public int HaveRevision;
+
+		/// <summary>
+		/// Changelist description (if using -e changelist and if the file was part of changelist)
+		/// </summary>
+		[PerforceTag("desc", Optional = true)]
+		public string Description;
+
+		/// <summary>
+		/// MD5 digest of a file (requires -Ol option)
+		/// </summary>
+		[PerforceTag("digest", Optional = true)]
+		public string Digest;
+
+		/// <summary>
+		/// File length in bytes (requires -Ol option)
+		/// </summary>
+		[PerforceTag("fileSize", Optional = true)]
+		public long FileSize;
+
+		/// <summary>
+		/// Open action, if opened in your workspace
+		/// </summary>
+		[PerforceTag("action", Optional = true)]
+		public FileAction Action;
+
+		/// <summary>
+		/// Open type, if opened in your workspace
+		/// </summary>
+		[PerforceTag("type", Optional = true)]
+		public FileAction Type;
+
+		/// <summary>
+		/// User who opened the file, if opens
+		/// </summary>
+		[PerforceTag("actionOwner", Optional = true)]
+		public string ActionOwner;
+
+		/// <summary>
+		/// Open changelist number, if opened in your workspace
+		/// </summary>
+		[PerforceTag("change", Optional = true)]
+		public int ChangeNumber;
+
+		/// <summary>
+		/// Summarize this record for display in the debugger
+		/// </summary>
+		/// <returns>Formatted record</returns>
+		public override string ToString()
+		{
+			return DepotFile ?? ClientFile ?? Path ?? base.ToString();
+		}
+	}
+}
diff --git a/Engine/Source/Programs/DotNETCommon/DotNETUtilities/Perforce/Records/FileAction.cs b/Engine/Source/Programs/DotNETCommon/DotNETUtilities/Perforce/Records/FileAction.cs
new file mode 100644
--- /dev/null
+++ b/Engine/Source/Programs/DotNETCommon/DotNETUtilities/Perforce/Records/FileAction.cs
@@ -0,0 +1,86 @@
+using System;
+using System.Collections.Generic;
+using System.Linq;
+using System.Text;
+using System.Threading.Tasks;
+
+namespace Tools.DotNETCommon.Perforce
+{
+	/// <summary>
+	/// Actions applied to a file in a particular revision
+	/// </summary>
+	public enum FileAction
+	{
+		/// <summary>
+		/// 
+		/// </summary>
+		[PerforceEnum("none")]
+		None,
+
+		/// <summary>
+		/// The file was added
+		/// </summary>
+		[PerforceEnum("add")]
+		Add,
+
+		/// <summary>
+		/// The file has been edited
+		/// </summary>
+		[PerforceEnum("edit")]
+		Edit,
+
+		/// <summary>
+		/// The file was deleted
+		/// </summary>
+		[PerforceEnum("delete")]
+		Delete,
+
+		/// <summary>
+		/// The file was branched from elsewhere in the depot
+		/// </summary>
+		[PerforceEnum("branch")]
+		Branch,
+
+		/// <summary>
+		/// The file at this path was added as part of a move from another location
+		/// </summary>
+		[PerforceEnum("move/add")]
+		MoveAdd,
+
+		/// <summary>
+		/// The file at this path was deleted as part of a move into another location
+		/// </summary>
+		[PerforceEnum("move/delete")]
+		MoveDelete,
+
+		/// <summary>
+		/// The file was merged with another file in the depot
+		/// </summary>
+		[PerforceEnum("integrate")]
+		Integrate,
+
+		/// <summary>
+		/// 
+		/// </summary>
+		[PerforceEnum("import")]
+		Import,
+
+		/// <summary>
+		/// The file was purged from the server
+		/// </summary>
+		[PerforceEnum("purge")]
+		Purge,
+
+		/// <summary>
+		/// The file was archived
+		/// </summary>
+		[PerforceEnum("archive")]
+		Archive,
+
+		/// <summary>
+		/// Unknown
+		/// </summary>
+		[PerforceEnum("unknown")]
+		Unknown,
+	}
+}
diff --git a/Engine/Source/Programs/DotNETCommon/DotNETUtilities/Perforce/Records/FileLogOptions.cs b/Engine/Source/Programs/DotNETCommon/DotNETUtilities/Perforce/Records/FileLogOptions.cs
new file mode 100644
--- /dev/null
+++ b/Engine/Source/Programs/DotNETCommon/DotNETUtilities/Perforce/Records/FileLogOptions.cs
@@ -0,0 +1,50 @@
+using System;
+using System.Collections.Generic;
+using System.Linq;
+using System.Text;
+using System.Threading.Tasks;
+
+namespace Tools.DotNETCommon.Perforce
+{
+	/// <summary>
+	/// Options for the filelog command
+	/// </summary>
+	[Flags]
+	public enum FileLogOptions
+	{
+		/// <summary>
+		/// No options
+		/// </summary>
+		None = 0,
+
+		/// <summary>
+		/// Display file content history instead of file name history.
+		/// </summary>
+		ContentHistory = 1,
+
+		/// <summary>
+		/// Follow file history across branches.
+		/// </summary>
+		FollowAcrossBranches = 2,
+
+		/// <summary>
+		/// List long output, with the full text of each changelist description.
+		/// </summary>
+		FullDescriptions = 4,
+
+		/// <summary>
+		/// List long output, with the full text of each changelist description truncated at 250 characters.
+		/// </summary>
+		LongDescriptions = 8,
+
+		/// <summary>
+		/// When used with the ContentHistory option, do not follow content of promoted task streams. 
+		/// </summary>
+		DoNotFollowPromotedTaskStreams = 16,
+
+		/// <summary>
+		/// Display a shortened form of output by ignoring non-contributory integrations
+		/// </summary>
+		IgnoreNonContributoryIntegrations = 32,
+	}
+}
diff --git a/Engine/Source/Programs/DotNETCommon/DotNETUtilities/Perforce/Records/FileLogRecord.cs b/Engine/Source/Programs/DotNETCommon/DotNETUtilities/Perforce/Records/FileLogRecord.cs
new file mode 100644
--- /dev/null
+++ b/Engine/Source/Programs/DotNETCommon/DotNETUtilities/Perforce/Records/FileLogRecord.cs
@@ -0,0 +1,35 @@
+using System;
+using System.Collections.Generic;
+using System.Linq;
+using System.Text;
+using System.Threading.Tasks;
+
+namespace Tools.DotNETCommon.Perforce
+{
+	/// <summary>
+	/// Record output by the filelog command
+	/// </summary>
+	public class FileLogRecord
+	{
+		/// <summary>
+		/// Path to the file in the depot
+		/// </summary>
+		[PerforceTag("depotFile")]
+		public string DepotPath;
+
+		/// <summary>
+		/// Revisions of this file
+		/// </summary>
+		[PerforceRecordList]
+		public List<RevisionRecord> Revisions = new List<RevisionRecord>();
+
+		/// <summary>
+		/// Format this record for display in the debugger
+		/// </summary>
+		/// <returns>Summary of this revision</returns>
+		public override string ToString()
+		{
+			return DepotPath;
+		}
+	}
+}
diff --git a/Engine/Source/Programs/DotNETCommon/DotNETUtilities/Perforce/Records/GetChangeOptions.cs b/Engine/Source/Programs/DotNETCommon/DotNETUtilities/Perforce/Records/GetChangeOptions.cs
new file mode 100644
--- /dev/null
+++ b/Engine/Source/Programs/DotNETCommon/DotNETUtilities/Perforce/Records/GetChangeOptions.cs
@@ -0,0 +1,25 @@
+using System;
+using System.Collections.Generic;
+using System.Linq;
+using System.Text;
+using System.Threading.Tasks;
+
+namespace Tools.DotNETCommon.Perforce
+{
+	/// <summary>
+	/// Options for the GetChange command (p4 change -o)
+	/// </summary>
+	public enum GetChangeOptions
+	{
+		/// <summary>
+		/// No options are specified
+		/// </summary>
+		None = 0,
+
+		/// <summary>
+		/// If a changelist was renumbered on submit, and you know only the original changelist number, use OriginalChangeNumber and the original changelist number to view or edit the changelist.
+		/// </summary>
+		BeforeRenumber = 1,
+	}
+
+}
diff --git a/Engine/Source/Programs/DotNETCommon/DotNETUtilities/Perforce/Records/InfoOptions.cs b/Engine/Source/Programs/DotNETCommon/DotNETUtilities/Perforce/Records/InfoOptions.cs
new file mode 100644
--- /dev/null
+++ b/Engine/Source/Programs/DotNETCommon/DotNETUtilities/Perforce/Records/InfoOptions.cs
@@ -0,0 +1,25 @@
+using System;
+using System.Collections.Generic;
+using System.Linq;
+using System.Text;
+using System.Threading.Tasks;
+
+namespace Tools.DotNETCommon.Perforce
+{
+	/// <summary>
+	/// Options for the p4 info command
+	/// </summary>
+	[Flags]
+	public enum InfoOptions
+	{
+		/// <summary>
+		/// No addiional options
+		/// </summary>
+		None = 0,
+
+		/// <summary>
+		/// Shortened output; exclude information that requires a database lookup.
+		/// </summary>
+		ShortOutput = 1,
+	}
+}
diff --git a/Engine/Source/Programs/DotNETCommon/DotNETUtilities/Perforce/Records/InfoRecord.cs b/Engine/Source/Programs/DotNETCommon/DotNETUtilities/Perforce/Records/InfoRecord.cs
new file mode 100644
--- /dev/null
+++ b/Engine/Source/Programs/DotNETCommon/DotNETUtilities/Perforce/Records/InfoRecord.cs
@@ -0,0 +1,56 @@
+using System;
+using System.Collections.Generic;
+using System.Linq;
+using System.Text;
+using System.Threading.Tasks;
+
+namespace Tools.DotNETCommon.Perforce
+{
+	/// <summary>
+	/// Contains information about the current client and server configuration
+	/// </summary>
+	public class InfoRecord
+	{
+		/// <summary>
+		/// The current user according to the Perforce environment
+		/// </summary>
+		[PerforceTag("userName", Optional = true)]
+		public string UserName;
+
+		/// <summary>
+		/// The current client
+		/// </summary>
+		[PerforceTag("clientName", Optional = true)]
+		public string ClientName;
+
+		/// <summary>
+		/// The current host
+		/// </summary>
+		[PerforceTag("clientHost", Optional = true)]
+		public string ClientHost;
+
+		/// <summary>
+		/// Root directory for the current client
+		/// </summary>
+		[PerforceTag("clientRoot", Optional = true)]
+		public string ClientRoot;
+
+		/// <summary>
+		/// Selected stream in the current client
+		/// </summary>
+		[PerforceTag("clientStream", Optional = true)]
+		public string ClientStream;
+
+		/// <summary>
+		/// Address of the Perforce server
+		/// </summary>
+		[PerforceTag("serverAddress", Optional = true)]
+		public string ServerAddress;
+
+		/// <summary>
+		/// Case handling setting on the server
+		/// </summary>
+		[PerforceTag("caseHandling", Optional = true)]
+		public string CaseHandling;
+	}
+}
diff --git a/Engine/Source/Programs/DotNETCommon/DotNETUtilities/Perforce/Records/IntegrationAction.cs b/Engine/Source/Programs/DotNETCommon/DotNETUtilities/Perforce/Records/IntegrationAction.cs
new file mode 100644
--- /dev/null
+++ b/Engine/Source/Programs/DotNETCommon/DotNETUtilities/Perforce/Records/IntegrationAction.cs
@@ -0,0 +1,110 @@
+using System;
+using System.Collections.Generic;
+using System.Linq;
+using System.Text;
+using System.Threading.Tasks;
+
+namespace Tools.DotNETCommon.Perforce
+{
+	/// <summary>
+	/// Describes the action performed by the user when resolving the integration
+	/// </summary>
+	public enum IntegrateAction
+	{
+		/// <summary>
+		/// file did not previously exist; it was created as a copy of partner-file
+		/// </summary>
+		[PerforceEnum("branch from")]
+		BranchFrom,
+
+		/// <summary>
+		/// partner-file did not previously exist; it was created as a copy of file.
+		/// </summary>
+		[PerforceEnum("branch into")]
+		BranchInto,
+
+		/// <summary>
+		/// file was integrated from partner-file, accepting merge.
+		/// </summary>
+		[PerforceEnum("merge from")]
+		MergeFrom,
+
+		/// <summary>
+		/// file was integrated into partner-file, accepting merge.
+		/// </summary>
+		[PerforceEnum("merge into")]
+		MergeInto,
+
+		/// <summary>
+		/// file was integrated from partner-file, accepting theirs and deleting the original.
+		/// </summary>
+		[PerforceEnum("moved from")]
+		MovedFrom,
+
+		/// <summary>
+		/// file was integrated into partner-file, accepting theirs and creating partner-file if it did not previously exist.
+		/// </summary>
+		[PerforceEnum("moved into")]
+		MovedInto,
+
+		/// <summary>
+		/// file was integrated from partner-file, accepting theirs.
+		/// </summary>
+		[PerforceEnum("copy from")]
+		CopyFrom,
+
+		/// <summary>
+		/// file was integrated into partner-file, accepting theirs.
+		/// </summary>
+		[PerforceEnum("copy into")]
+		CopyInto,
+
+		/// <summary>
+		/// file was integrated from partner-file, accepting yours.
+		/// </summary>
+		[PerforceEnum("ignored")]
+		Ignored,
+
+		/// <summary>
+		/// file was integrated into partner-file, accepting yours.
+		/// </summary>
+		[PerforceEnum("ignored by")]
+		IgnoredBy,
+
+		/// <summary>
+		/// file was integrated from partner-file, and partner-file had been previously deleted.
+		/// </summary>
+		[PerforceEnum("delete from")]
+		DeleteFrom,
+
+		/// <summary>
+		/// file was integrated into partner-file, and file had been previously deleted.
+		/// </summary>
+		[PerforceEnum("delete into")]
+		DeleteInto,
+
+		/// <summary>
+		/// file was integrated from partner-file, and file was edited within the p4 resolve process.
+		/// </summary>
+		[PerforceEnum("edit from")]
+		EditFrom,
+
+		/// <summary>
+		/// file was integrated into partner-file, and partner-file was reopened for edit before submission.
+		/// </summary>
+		[PerforceEnum("edit into")]
+		EditInto,
+
+		/// <summary>
+		/// file was integrated from a deleted partner-file, and partner-file was reopened for add (that is, someone restored a deleted file by syncing back to a pre-deleted revision and adding the file).
+		/// </summary>
+		[PerforceEnum("add from")]
+		AddFrom,
+
+		/// <summary>
+		/// file was integrated into previously nonexistent partner-file, and partner-file was reopened for add before submission.
+		/// </summary>
+		[PerforceEnum("add into")]
+		AddInto,
+	}
+}
diff --git a/Engine/Source/Programs/DotNETCommon/DotNETUtilities/Perforce/Records/IntegrationRecord.cs b/Engine/Source/Programs/DotNETCommon/DotNETUtilities/Perforce/Records/IntegrationRecord.cs
new file mode 100644
--- /dev/null
+++ b/Engine/Source/Programs/DotNETCommon/DotNETUtilities/Perforce/Records/IntegrationRecord.cs
@@ -0,0 +1,54 @@
+using System;
+using System.Collections.Generic;
+using System.Linq;
+using System.Text;
+using System.Threading.Tasks;
+
+namespace Tools.DotNETCommon.Perforce
+{
+	/// <summary>
+	/// Stores integration information for a file revision
+	/// </summary>
+	public class IntegrationRecord
+	{
+		/// <summary>
+		/// The integration action performed for this file
+		/// </summary>
+		[PerforceTag("how")]
+		public IntegrateAction Action;
+
+		/// <summary>
+		/// The partner file for this integration
+		/// </summary>
+		[PerforceTag("file")]
+		public string OtherFile;
+
+		/// <summary>
+		/// Min revision of the partner file for this integration
+		/// </summary>
+		[PerforceTag("srev")]
+		public int StartRevisionNumber;
+
+		/// <summary>
+		/// Max revision of the partner file for this integration
+		/// </summary>
+		[PerforceTag("erev")]
+		public int EndRevisionNumber;
+
+		/// <summary>
+		/// Summarize this record for display in the debugger
+		/// </summary>
+		/// <returns>Formatted integration record</returns>
+		public override string ToString()
+		{
+			if(StartRevisionNumber + 1 == EndRevisionNumber)
+			{
+				return String.Format("{0} {1}#{2}", Action, OtherFile, EndRevisionNumber);
+			}
+			else
+			{
+				return String.Format("{0} {1}#{2},#{3}", Action, OtherFile, StartRevisionNumber + 1, EndRevisionNumber);
+			}
+		}
+	}
+}
diff --git a/Engine/Source/Programs/DotNETCommon/DotNETUtilities/Perforce/Records/PrintRecord.cs b/Engine/Source/Programs/DotNETCommon/DotNETUtilities/Perforce/Records/PrintRecord.cs
new file mode 100644
--- /dev/null
+++ b/Engine/Source/Programs/DotNETCommon/DotNETUtilities/Perforce/Records/PrintRecord.cs
@@ -0,0 +1,56 @@
+using System;
+using System.Collections.Generic;
+using System.Linq;
+using System.Text;
+using System.Threading.Tasks;
+
+namespace Tools.DotNETCommon.Perforce
+{
+	/// <summary>
+	/// Record containing information about a printed file
+	/// </summary>
+	public class PrintRecord
+	{
+		/// <summary>
+		/// Path to the file in the depot
+		/// </summary>
+		[PerforceTag("depotFile")]
+		public string DepotFile;
+
+		/// <summary>
+		/// Revision number of the file
+		/// </summary>
+		[PerforceTag("rev")]
+		public int Revision;
+
+		/// <summary>
+		/// Change number of the file
+		/// </summary>
+		[PerforceTag("change")]
+		public int ChangeNumber;
+
+		/// <summary>
+		/// Last action to the file
+		/// </summary>
+		[PerforceTag("action")]
+		public FileAction Action;
+
+		/// <summary>
+		/// File type
+		/// </summary>
+		[PerforceTag("type")]
+		public string Type;
+
+		/// <summary>
+		/// Submit time of the file
+		/// </summary>
+		[PerforceTag("time")]
+		public DateTime Time;
+
+		/// <summary>
+		/// Size of the file in bytes
+		/// </summary>
+		[PerforceTag("fileSize")]
+		public long FileSize;
+	}
+}
diff --git a/Engine/Source/Programs/DotNETCommon/DotNETUtilities/Perforce/Records/ReconcileOptions.cs b/Engine/Source/Programs/DotNETCommon/DotNETUtilities/Perforce/Records/ReconcileOptions.cs
new file mode 100644
--- /dev/null
+++ b/Engine/Source/Programs/DotNETCommon/DotNETUtilities/Perforce/Records/ReconcileOptions.cs
@@ -0,0 +1,55 @@
+using System;
+using System.Collections.Generic;
+using System.Linq;
+using System.Text;
+using System.Threading.Tasks;
+
+namespace Tools.DotNETCommon.Perforce
+{
+	/// <summary>
+	/// Options for the p4 reconcile command
+	/// </summary>
+	[Flags]
+	public enum ReconcileOptions
+	{
+		/// <summary>
+		/// No options specified
+		/// </summary>
+		None = 0,
+
+		/// <summary>
+		/// Find files in the client workspace that have been modified outside of Perforce, and open them for edit.
+		/// </summary>
+		Edit = 1,
+
+		/// <summary>
+		/// Find files in the client workspace that are not under Perforce control and open them for add.
+		/// </summary>
+		Add = 2,
+
+		/// <summary>
+		/// Find files missing from the client workspace, but present on the server; open these files for delete, but only if these files are in the user's have list.
+		/// </summary>
+		Delete = 4,
+
+		/// <summary>
+		/// Preview which files would be opened for add, without actually changing any files or metadata.
+		/// </summary>
+		PreviewOnly = 8,
+
+		/// <summary>
+		/// Add filenames that contain special (wildcard) characters.
+		/// </summary>
+		AllowWildcards = 16,
+
+		/// <summary>
+		/// Do not perform any ignore checking.
+		/// </summary>
+		NoIgnore = 32,
+
+		/// <summary>
+		/// Display output in local file syntax with relative paths, similar to the workspace-centric view of p4 status.
+		/// </summary>
+		LocalFileSyntax = 64,
+	}
+}
diff --git a/Engine/Source/Programs/DotNETCommon/DotNETUtilities/Perforce/Records/ReconcileRecord.cs b/Engine/Source/Programs/DotNETCommon/DotNETUtilities/Perforce/Records/ReconcileRecord.cs
new file mode 100644
--- /dev/null
+++ b/Engine/Source/Programs/DotNETCommon/DotNETUtilities/Perforce/Records/ReconcileRecord.cs
@@ -0,0 +1,53 @@
+using System;
+using System.Collections.Generic;
+using System.Linq;
+using System.Text;
+using System.Threading.Tasks;
+
+namespace Tools.DotNETCommon.Perforce
+{
+	/// <summary>
+	/// Information about a reconciled file
+	/// </summary>
+	public class ReconcileRecord
+	{
+		/// <summary>
+		/// Path to the file in the depot
+		/// </summary>
+		[PerforceTag("depotFile")]
+		public string DepotFile;
+
+		/// <summary>
+		/// Path to the file in the workspace
+		/// </summary>
+		[PerforceTag("clientFile")]
+		public string ClientFile;
+
+		/// <summary>
+		/// The working revision number of the file that was synced
+		/// </summary>
+		[PerforceTag("workRev")]
+		public int WorkingRevision;
+
+		/// <summary>
+		/// Action taken when syncing the file
+		/// </summary>
+		[PerforceTag("action")]
+		public string Action;
+
+		/// <summary>
+		/// Type of the file
+		/// </summary>
+		[PerforceTag("type")]
+		public string Type;
+
+		/// <summary>
+		/// Format this record for display in the debugger
+		/// </summary>
+		/// <returns>Summary of this revision</returns>
+		public override string ToString()
+		{
+			return DepotFile;
+		}
+	}
+}
diff --git a/Engine/Source/Programs/DotNETCommon/DotNETUtilities/Perforce/Records/ResolveOptions.cs b/Engine/Source/Programs/DotNETCommon/DotNETUtilities/Perforce/Records/ResolveOptions.cs
new file mode 100644
--- /dev/null
+++ b/Engine/Source/Programs/DotNETCommon/DotNETUtilities/Perforce/Records/ResolveOptions.cs
@@ -0,0 +1,70 @@
+using System;
+using System.Collections.Generic;
+using System.Linq;
+using System.Text;
+using System.Threading.Tasks;
+
+namespace Tools.DotNETCommon.Perforce
+{
+	/// <summary>
+	/// Options for the p4 resolve command
+	/// </summary>
+	[Flags]
+	public enum ResolveOptions
+	{
+		/// <summary>
+		/// No options specified
+		/// </summary>
+		None = 0,
+
+		/// <summary>
+		/// Automatically accept the Perforce-recommended file revision.
+		/// </summary>
+		Automatic = 1,
+
+		/// <summary>
+		/// Accept Yours, ignore theirs.
+		/// </summary>
+		AcceptYours = 2,
+
+		/// <summary>
+		/// Accept Theirs.
+		/// </summary>
+		AcceptTheirs = 4,
+
+		/// <summary>
+		/// If either yours or theirs is different from base, (and the changes are in common) accept that revision. If both are different from base, skip this file.
+		/// </summary>
+		SafeAccept = 8,
+
+		/// <summary>
+		/// Automatically accept the Perforce-recommended file revision.
+		/// </summary>
+		ForceAccept = 16,
+
+		/// <summary>
+		/// Ignore whitespace-only changes (for instance, a tab replaced by eight spaces)
+		/// </summary>
+		IgnoreWhitespaceOnly = 32,
+
+		/// <summary>
+		/// Ignore whitespace altogether (for instance, deletion of tabs or other whitespace)
+		/// </summary>
+		IgnoreWhitespace = 64,
+
+		/// <summary>
+		/// Ignore differences in line-ending convention
+		/// </summary>
+		IgnoreLineEndings = 128,
+
+		/// <summary>
+		/// Allow already resolved, but not yet submitted, files to be resolved again.
+		/// </summary>
+		ResolveAgain = 256,
+
+		/// <summary>
+		/// List the files that need resolving without actually performing the resolve.
+		/// </summary>
+		PreviewOnly = 512,
+	}
+}
diff --git a/Engine/Source/Programs/DotNETCommon/DotNETUtilities/Perforce/Records/ResolveRecord.cs b/Engine/Source/Programs/DotNETCommon/DotNETUtilities/Perforce/Records/ResolveRecord.cs
new file mode 100644
--- /dev/null
+++ b/Engine/Source/Programs/DotNETCommon/DotNETUtilities/Perforce/Records/ResolveRecord.cs
@@ -0,0 +1,59 @@
+using System;
+using System.Collections.Generic;
+using System.Linq;
+using System.Text;
+using System.Threading.Tasks;
+
+namespace Tools.DotNETCommon.Perforce
+{
+	/// <summary>
+	/// Information about a resolved file
+	/// </summary>
+	public class ResolveRecord
+	{
+		/// <summary>
+		/// Path to the file in the workspace
+		/// </summary>
+		[PerforceTag("clientFile")]
+		public string ClientFile;
+
+		/// <summary>
+		/// Path to the depot file that needs to be resolved
+		/// </summary>
+		[PerforceTag("fromFile")]
+		public string FromFile;
+
+		/// <summary>
+		/// Start range of changes to be resolved
+		/// </summary>
+		[PerforceTag("startFromRev")]
+		public int FromRevisionStart;
+
+		/// <summary>
+		/// Ending range of changes to be resolved
+		/// </summary>
+		[PerforceTag("endFromRev")]
+		public int FromRevisionEnd;
+
+		/// <summary>
+		/// The type of resolve to perform
+		/// </summary>
+		[PerforceTag("resolveType")]
+		public string ResolveType;
+
+		/// <summary>
+		/// For content resolves, the type of resolve to be performed
+		/// </summary>
+		[PerforceTag("contentResolveType", Optional = true)]
+		public string ContentResolveType;
+
+		/// <summary>
+		/// Format this record for display in the debugger
+		/// </summary>
+		/// <returns>Summary of this revision</returns>
+		public override string ToString()
+		{
+			return ClientFile;
+		}
+	}
+}
diff --git a/Engine/Source/Programs/DotNETCommon/DotNETUtilities/Perforce/Records/RevertOptions.cs b/Engine/Source/Programs/DotNETCommon/DotNETUtilities/Perforce/Records/RevertOptions.cs
new file mode 100644
--- /dev/null
+++ b/Engine/Source/Programs/DotNETCommon/DotNETUtilities/Perforce/Records/RevertOptions.cs
@@ -0,0 +1,40 @@
+using System;
+using System.Collections.Generic;
+using System.Linq;
+using System.Text;
+using System.Threading.Tasks;
+
+namespace Tools.DotNETCommon.Perforce
+{
+	/// <summary>
+	/// Options for the revert command
+	/// </summary>
+	[Flags]
+	public enum RevertOptions
+	{
+		/// <summary>
+		/// No options specified
+		/// </summary>
+		None = 0,
+
+		/// <summary>
+		/// Revert only those files that haven’t changed (in terms of content or filetype) since they were opened.
+		/// </summary>
+		Unchanged = 1,
+
+		/// <summary>
+		/// Keep workspace files; the file(s) are removed from any changelists and Perforce records that the files as being no longer open, but the file(s) are unchanged in the client workspace.
+		/// </summary>
+		KeepWorkspaceFiles = 2,
+
+		/// <summary>
+		/// List the files that would be reverted without actually performing the revert.
+		/// </summary>
+		PreviewOnly = 4,
+
+		/// <summary>
+		/// Files that are open for add are to be deleted (wiped) from the workspace when reverted.
+		/// </summary>
+		DeleteAddedFiles = 8,
+	}
+}
diff --git a/Engine/Source/Programs/DotNETCommon/DotNETUtilities/Perforce/Records/RevertRecord.cs b/Engine/Source/Programs/DotNETCommon/DotNETUtilities/Perforce/Records/RevertRecord.cs
new file mode 100644
--- /dev/null
+++ b/Engine/Source/Programs/DotNETCommon/DotNETUtilities/Perforce/Records/RevertRecord.cs
@@ -0,0 +1,53 @@
+using System;
+using System.Collections.Generic;
+using System.Linq;
+using System.Text;
+using System.Threading.Tasks;
+
+namespace Tools.DotNETCommon.Perforce
+{
+	/// <summary>
+	/// Information about a reverted file
+	/// </summary>
+	public class RevertRecord
+	{
+		/// <summary>
+		/// Path to the file in the depot
+		/// </summary>
+		[PerforceTag("depotFile")]
+		public string DepotFile;
+
+		/// <summary>
+		/// Path to the file in the workspace
+		/// </summary>
+		[PerforceTag("clientFile")]
+		public string ClientFile;
+
+		/// <summary>
+		/// The revision number that is now in the workspace
+		/// </summary>
+		[PerforceTag("haveRev")]
+		public int HaveRevision;
+
+		/// <summary>
+		/// The previous action that the file was opened as
+		/// </summary>
+		[PerforceTag("oldAction")]
+		public string OldAction;
+
+		/// <summary>
+		/// Action taken to revert the file
+		/// </summary>
+		[PerforceTag("action")]
+		public string Action;
+
+		/// <summary>
+		/// Format this record for display in the debugger
+		/// </summary>
+		/// <returns>Summary of this revert record</returns>
+		public override string ToString()
+		{
+			return ClientFile;
+		}
+	}
+}
diff --git a/Engine/Source/Programs/DotNETCommon/DotNETUtilities/Perforce/Records/RevisionRecord.cs b/Engine/Source/Programs/DotNETCommon/DotNETUtilities/Perforce/Records/RevisionRecord.cs
new file mode 100644
--- /dev/null
+++ b/Engine/Source/Programs/DotNETCommon/DotNETUtilities/Perforce/Records/RevisionRecord.cs
@@ -0,0 +1,89 @@
+using System;
+using System.Collections.Generic;
+using System.Linq;
+using System.Text;
+using System.Threading.Tasks;
+
+namespace Tools.DotNETCommon.Perforce
+{
+	/// <summary>
+	/// Stores a revision record for a file
+	/// </summary>
+	public class RevisionRecord
+	{
+		/// <summary>
+		/// The revision number of this file
+		/// </summary>
+		[PerforceTag("rev")]
+		public readonly int RevisionNumber;
+
+		/// <summary>
+		/// The changelist responsible for this revision of the file
+		/// </summary>
+		[PerforceTag("change")]
+		public readonly int ChangeNumber;
+
+		/// <summary>
+		/// Action performed to the file in this revision
+		/// </summary>
+		[PerforceTag("action")]
+		public readonly FileAction Action;
+
+		/// <summary>
+		/// Type of the file
+		/// </summary>
+		[PerforceTag("type")]
+		public readonly string Type;
+
+		/// <summary>
+		/// Timestamp of this modification
+		/// </summary>
+		[PerforceTag("time")]
+		public readonly DateTime Time;
+
+		/// <summary>
+		/// Author of the changelist
+		/// </summary>
+		[PerforceTag("user")]
+		public readonly string UserName;
+
+		/// <summary>
+		/// Client that submitted this changelist
+		/// </summary>
+		[PerforceTag("client")]
+		public readonly string ClientName;
+
+		/// <summary>
+		/// Size of the file, or -1 if not specified
+		/// </summary>
+		[PerforceTag("fileSize", Optional = true)]
+		public readonly int FileSize;
+
+		/// <summary>
+		/// Digest of the file, or null if not specified
+		/// </summary>
+		[PerforceTag("digest", Optional = true)]
+		public readonly string Digest;
+
+		/// <summary>
+		/// Description of this changelist
+		/// </summary>
+		[PerforceTag("desc")]
+		public string Description;
+
+		/// <summary>
+		/// Integration records for this revision
+		/// </summary>
+		[PerforceRecordList]
+		public List<IntegrationRecord> Integrations = new List<IntegrationRecord>();
+
+		/// <summary>
+		/// Format this record for display in the debugger
+		/// </summary>
+		/// <returns>Summary of this revision</returns>
+		public override string ToString()
+		{
+			return String.Format("#{0} change {1} {2} on {3} by {4}@{5}", RevisionNumber, ChangeNumber, Action, Time, UserName, ClientName);
+		}
+	}
+}
diff --git a/Engine/Source/Programs/DotNETCommon/DotNETUtilities/Perforce/Records/ShelveOptions.cs b/Engine/Source/Programs/DotNETCommon/DotNETUtilities/Perforce/Records/ShelveOptions.cs
new file mode 100644
--- /dev/null
+++ b/Engine/Source/Programs/DotNETCommon/DotNETUtilities/Perforce/Records/ShelveOptions.cs
@@ -0,0 +1,30 @@
+using System;
+using System.Collections.Generic;
+using System.Linq;
+using System.Text;
+using System.Threading.Tasks;
+
+namespace Tools.DotNETCommon.Perforce
+{
+	/// <summary>
+	/// Options for the p4 shelve command
+	/// </summary>
+	[Flags]
+	public enum ShelveOptions
+	{
+		/// <summary>
+		/// No options specified
+		/// </summary>
+		None = 0,
+
+		/// <summary>
+		/// Only shelve files that have changed
+		/// </summary>
+		OnlyChanged = 1,
+
+		/// <summary>
+		/// Force the overwriting of any existing shelved files in a pending changelist with the contents of their client workspace copies.
+		/// </summary>
+		Overwrite = 2,
+	}
+}
diff --git a/Engine/Source/Programs/DotNETCommon/DotNETUtilities/Perforce/Records/ShelveRecord.cs b/Engine/Source/Programs/DotNETCommon/DotNETUtilities/Perforce/Records/ShelveRecord.cs
new file mode 100644
--- /dev/null
+++ b/Engine/Source/Programs/DotNETCommon/DotNETUtilities/Perforce/Records/ShelveRecord.cs
@@ -0,0 +1,47 @@
+using System;
+using System.Collections.Generic;
+using System.Linq;
+using System.Text;
+using System.Threading.Tasks;
+
+namespace Tools.DotNETCommon.Perforce
+{
+	/// <summary>
+	/// Contains information about a shelved file
+	/// </summary>
+	public class ShelveRecord
+	{
+		/// <summary>
+		/// The changelist containing the shelved file
+		/// </summary>
+		[PerforceTag("change")]
+		public int Change;
+
+		/// <summary>
+		/// Path to the file in the depot
+		/// </summary>
+		[PerforceTag("depotFile")]
+		public string DepotFile;
+
+		/// <summary>
+		/// The revision number of the file that was shelved
+		/// </summary>
+		[PerforceTag("rev")]
+		public int Revision;
+
+		/// <summary>
+		/// The action to be applied to the file
+		/// </summary>
+		[PerforceTag("action")]
+		public string Action;
+
+		/// <summary>
+		/// Format this record for display in the debugger
+		/// </summary>
+		/// <returns>Summary of this revision</returns>
+		public override string ToString()
+		{
+			return DepotFile;
+		}
+	}
+}
diff --git a/Engine/Source/Programs/DotNETCommon/DotNETUtilities/Perforce/Records/StreamOptions.cs b/Engine/Source/Programs/DotNETCommon/DotNETUtilities/Perforce/Records/StreamOptions.cs
new file mode 100644
--- /dev/null
+++ b/Engine/Source/Programs/DotNETCommon/DotNETUtilities/Perforce/Records/StreamOptions.cs
@@ -0,0 +1,75 @@
+using System;
+using System.Collections.Generic;
+using System.Linq;
+using System.Text;
+using System.Threading.Tasks;
+
+namespace Tools.DotNETCommon.Perforce
+{
+	/// <summary>
+	/// Options for a stream definition
+	/// </summary>
+	[Flags]
+	public enum StreamOptions
+	{
+		/// <summary>
+		/// The stream is locked
+		/// </summary>
+		[PerforceEnum("locked")]
+		Locked = 1,
+
+		/// <summary>
+		/// The stream is unlocked
+		/// </summary>
+		[PerforceEnum("unlocked")]
+		Unlocked = 2,
+
+		/// <summary>
+		/// Only the owner may submit to the stream
+		/// </summary>
+		[PerforceEnum("ownersubmit")]
+		OwnerSubmit = 4,
+
+		/// <summary>
+		/// Anyone may submit to the stream
+		/// </summary>
+		[PerforceEnum("allsubmit")]
+		AllSubmit = 8,
+
+		/// <summary>
+		/// Integrations from this stream to its parent are expected
+		/// </summary>
+		[PerforceEnum("toparent")]
+		ToParent = 16,
+
+		/// <summary>
+		/// Integrations from this stream from its parent are expected
+		/// </summary>
+		[PerforceEnum("notoparent")]
+		NotToParent = 32,
+
+		/// <summary>
+		/// Integrations from this stream from its parent are expected
+		/// </summary>
+		[PerforceEnum("fromparent")]
+		FromParent = 64,
+
+		/// <summary>
+		/// Integrations from this stream from its parent are expected
+		/// </summary>
+		[PerforceEnum("nofromparent")]
+		NotFromParent = 128,
+
+		/// <summary>
+		/// Undocumented?
+		/// </summary>
+		[PerforceEnum("mergedown")]
+		MergeDown = 256,
+
+		/// <summary>
+		/// Undocumented?
+		/// </summary>
+		[PerforceEnum("mergeany")]
+		MergeAny = 512,
+	}
+}
diff --git a/Engine/Source/Programs/DotNETCommon/DotNETUtilities/Perforce/Records/StreamRecord.cs b/Engine/Source/Programs/DotNETCommon/DotNETUtilities/Perforce/Records/StreamRecord.cs
new file mode 100644
--- /dev/null
+++ b/Engine/Source/Programs/DotNETCommon/DotNETUtilities/Perforce/Records/StreamRecord.cs
@@ -0,0 +1,102 @@
+using System;
+using System.Collections.Generic;
+using System.Diagnostics;
+using System.Linq;
+using System.Text;
+using System.Threading.Tasks;
+
+namespace Tools.DotNETCommon.Perforce
+{
+	/// <summary>
+	/// Contains information about a stream, as returned by the 'p4 streams' command
+	/// </summary>
+	public class StreamRecord
+	{
+		/// <summary>
+		/// Path to the stream
+		/// </summary>
+		[PerforceTag("Stream")]
+		public string Stream;
+
+		/// <summary>
+		/// Last time the stream definition was updated
+		/// </summary>
+		[PerforceTag("Update")]
+		public DateTime Update;
+
+		/// <summary>
+		/// Last time the stream definition was accessed
+		/// </summary>
+		[PerforceTag("Access")]
+		public DateTime Access;
+
+		/// <summary>
+		/// Owner of this stream
+		/// </summary>
+		[PerforceTag("Owner")]
+		public string Owner;
+
+		/// <summary>
+		/// Name of the stream. This may be modified after the stream is initially created, but it's underlying depot path will not change.
+		/// </summary>
+		[PerforceTag("Name")]
+		public string Name;
+
+		/// <summary>
+		/// The parent stream
+		/// </summary>
+		[PerforceTag("Parent")]
+		public string Parent;
+
+		/// <summary>
+		/// Type of the stream
+		/// </summary>
+		[PerforceTag("Type")]
+		public StreamType Type;
+
+		/// <summary>
+		/// User supplied description of the stream
+		/// </summary>
+		[PerforceTag("desc")]
+		public string Description;
+
+		/// <summary>
+		/// Options for the stream definition
+		/// </summary>
+		[PerforceTag("Options")]
+		public StreamOptions Options;
+
+		/// <summary>
+		/// Whether this stream is more stable than the parent stream
+		/// </summary>
+		[PerforceTag("firmerThanParent")]
+		public Nullable<bool> FirmerThanParent;
+
+		/// <summary>
+		/// Whether changes from this stream flow to the parent stream
+		/// </summary>
+		[PerforceTag("changeFlowsToParent")]
+		public bool ChangeFlowsToParent;
+
+		/// <summary>
+		/// Whether changes from this stream flow from the parent stream
+		/// </summary>
+		[PerforceTag("changeFlowsFromParent")]
+		public bool ChangeFlowsFromParent;
+
+		/// <summary>
+		/// The mainline branch associated with this stream
+		/// </summary>
+		[PerforceTag("baseParent")]
+		public string BaseParent;
+
+		/// <summary>
+		/// Summarize this record for display in the debugger
+		/// </summary>
+		/// <returns>Formatted stream information</returns>
+		public override string ToString()
+		{
+			return Name;
+		}
+	}
+}
diff --git a/Engine/Source/Programs/DotNETCommon/DotNETUtilities/Perforce/Records/StreamType.cs b/Engine/Source/Programs/DotNETCommon/DotNETUtilities/Perforce/Records/StreamType.cs
new file mode 100644
--- /dev/null
+++ b/Engine/Source/Programs/DotNETCommon/DotNETUtilities/Perforce/Records/StreamType.cs
@@ -0,0 +1,44 @@
+using System;
+using System.Collections.Generic;
+using System.Linq;
+using System.Text;
+using System.Threading.Tasks;
+
+namespace Tools.DotNETCommon.Perforce
+{
+	/// <summary>
+	/// Type of a Perforce stream
+	/// </summary>
+	public enum StreamType
+	{
+		/// <summary>
+		/// A mainline stream
+		/// </summary>
+		[PerforceEnum("mainline")]
+		Mainline,
+
+		/// <summary>
+		/// A development stream
+		/// </summary>
+		[PerforceEnum("development")]
+		Development,
+
+		/// <summary>
+		/// A release stream
+		/// </summary>
+		[PerforceEnum("release")]
+		Release,
+
+		/// <summary>
+		/// A virtual stream
+		/// </summary>
+		[PerforceEnum("virtual")]
+		Virtual,
+
+		/// <summary>
+		/// A task stream
+		/// </summary>
+		[PerforceEnum("task")]
+		Task,
+	}
+}
diff --git a/Engine/Source/Programs/DotNETCommon/DotNETUtilities/Perforce/Records/SubmitOptions.cs b/Engine/Source/Programs/DotNETCommon/DotNETUtilities/Perforce/Records/SubmitOptions.cs
new file mode 100644
--- /dev/null
+++ b/Engine/Source/Programs/DotNETCommon/DotNETUtilities/Perforce/Records/SubmitOptions.cs
@@ -0,0 +1,25 @@
+using System;
+using System.Collections.Generic;
+using System.Linq;
+using System.Text;
+using System.Threading.Tasks;
+
+namespace Tools.DotNETCommon.Perforce
+{
+	/// <summary>
+	/// Options for the 'submit' command
+	/// </summary>
+	[Flags]
+	public enum SubmitOptions
+	{
+		/// <summary>
+		/// No options specified
+		/// </summary>
+		None = 0,
+
+		/// <summary>
+		/// Reopen files for edit in the default changelist after submission.
+		/// </summary>
+		ReopenAsEdit = 1,
+	}
+}
diff --git a/Engine/Source/Programs/DotNETCommon/DotNETUtilities/Perforce/Records/SubmitRecord.cs b/Engine/Source/Programs/DotNETCommon/DotNETUtilities/Perforce/Records/SubmitRecord.cs
new file mode 100644
--- /dev/null
+++ b/Engine/Source/Programs/DotNETCommon/DotNETUtilities/Perforce/Records/SubmitRecord.cs
@@ -0,0 +1,29 @@
+using System;
+using System.Collections.Generic;
+using System.Linq;
+using System.Text;
+using System.Threading.Tasks;
+
+namespace Tools.DotNETCommon.Perforce
+{
+	/// <summary>
+	/// Contains information about a submitted changelist
+	/// </summary>
+	public class SubmitRecord
+	{
+		/// <summary>
+		/// Submitted changelist number
+		/// </summary>
+		[PerforceTag("change", Optional = true)]
+		public int ChangeNumber = -1;
+
+		/// <summary>
+		/// Format this record for display in the debugger
+		/// </summary>
+		/// <returns>Summary of this revision</returns>
+		public override string ToString()
+		{
+			return ChangeNumber.ToString();
+		}
+	}
+}
diff --git a/Engine/Source/Programs/DotNETCommon/DotNETUtilities/Perforce/Records/SwitchClientOptions.cs b/Engine/Source/Programs/DotNETCommon/DotNETUtilities/Perforce/Records/SwitchClientOptions.cs
new file mode 100644
--- /dev/null
+++ b/Engine/Source/Programs/DotNETCommon/DotNETUtilities/Perforce/Records/SwitchClientOptions.cs
@@ -0,0 +1,25 @@
+using System;
+using System.Collections.Generic;
+using System.Linq;
+using System.Text;
+using System.Threading.Tasks;
+
+namespace Tools.DotNETCommon.Perforce
+{
+	/// <summary>
+	/// Options for the p4 client -s command
+	/// </summary>
+	[Flags]
+	public enum SwitchClientOptions
+	{
+		/// <summary>
+		/// No options specified
+		/// </summary>
+		None = 0,
+
+		/// <summary>
+		/// Ignore files which are open for edit
+		/// </summary>
+		IgnoreOpenFiles = 1,
+	}
+}
diff --git a/Engine/Source/Programs/DotNETCommon/DotNETUtilities/Perforce/Records/SyncOptions.cs b/Engine/Source/Programs/DotNETCommon/DotNETUtilities/Perforce/Records/SyncOptions.cs
new file mode 100644
--- /dev/null
+++ b/Engine/Source/Programs/DotNETCommon/DotNETUtilities/Perforce/Records/SyncOptions.cs
@@ -0,0 +1,60 @@
+using System;
+using System.Collections.Generic;
+using System.Linq;
+using System.Text;
+using System.Threading.Tasks;
+
+namespace Tools.DotNETCommon.Perforce
+{
+	/// <summary>
+	/// Options for the p4 sync command
+	/// </summary>
+	[Flags]
+	public enum SyncOptions
+	{
+		/// <summary>
+		/// No options specified
+		/// </summary>
+		None = 0,
+
+		/// <summary>
+		/// Force the sync. Perforce performs the sync even if the client workspace already has the file at the specified revision. If the file is writable, it is overwritten.
+		/// </summary>
+		Force = 1,
+
+		/// <summary>
+		/// Keep existing workspace files; update the have list without updating the client workspace.
+		/// </summary>
+		KeepWorkspaceFiles = 2,
+
+		/// <summary>
+		/// For scripting purposes, perform the sync on a list of valid file arguments in full depot syntax with a valid revision number.
+		/// </summary>
+		FullDepotSyntax = 4,
+
+		/// <summary>
+		/// Preview mode: Display the results of the sync without actually performing the sync.
+		/// </summary>
+		PreviewOnly = 8,
+
+		/// <summary>
+		/// Preview mode: Display a summary of the expected network traffic associated with a sync, without performing the sync.
+		/// </summary>
+		NetworkPreviewOnly = 16,
+
+		/// <summary>
+		/// Populate a client workspace, but do not update the have list. Any file that is already synced or opened is bypassed with a warning message.
+		/// </summary>
+		DoNotUpdateHaveList = 32,
+
+		/// <summary>
+		/// Reopen files that are mapped to new locations in the depot, in the new location.
+		/// </summary>
+		ReopenMovedFiles = 64,
+
+		/// <summary>
+		/// Safe sync: Compare the content in your client workspace against what was last synced.
+		/// </summary>
+		Safe = 128,
+	}
+}
diff --git a/Engine/Source/Programs/DotNETCommon/DotNETUtilities/Perforce/Records/SyncRecord.cs b/Engine/Source/Programs/DotNETCommon/DotNETUtilities/Perforce/Records/SyncRecord.cs
new file mode 100644
--- /dev/null
+++ b/Engine/Source/Programs/DotNETCommon/DotNETUtilities/Perforce/Records/SyncRecord.cs
@@ -0,0 +1,71 @@
+using System;
+using System.Collections.Generic;
+using System.Linq;
+using System.Text;
+using System.Threading.Tasks;
+
+namespace Tools.DotNETCommon.Perforce
+{
+	/// <summary>
+	/// Information about a synced file
+	/// </summary>
+	public class SyncRecord
+	{
+		/// <summary>
+		/// Path to the file in the depot
+		/// </summary>
+		[PerforceTag("depotFile")]
+		public string DepotFile;
+
+		/// <summary>
+		/// Path to the file in the workspace
+		/// </summary>
+		[PerforceTag("clientFile")]
+		public string ClientFile;
+
+		/// <summary>
+		/// The revision number of the file that was synced
+		/// </summary>
+		[PerforceTag("rev")]
+		public int Revision;
+
+		/// <summary>
+		/// Action taken when syncing the file
+		/// </summary>
+		[PerforceTag("action")]
+		public string Action;
+
+		/// <summary>
+		/// Size of the file
+		/// </summary>
+		[PerforceTag("fileSize", Optional = true)]
+		public long FileSize;
+
+		/// <summary>
+		/// Stores the total size of all files that are being synced (only set for the first sync record)
+		/// </summary>
+		[PerforceTag("totalFileSize", Optional = true)]
+		public long TotalFileSize;
+
+		/// <summary>
+		/// Stores the total number of files that will be synced (only set for the first sync record)
+		/// </summary>
+		[PerforceTag("totalFileCount", Optional = true)]
+		public long TotalFileCount;
+
+		/// <summary>
+		/// Change that modified the file
+		/// </summary>
+		[PerforceTag("change", Optional = true)]
+		public int Change;
+
+		/// <summary>
+		/// Format this record for display in the debugger
+		/// </summary>
+		/// <returns>Summary of this revision</returns>
+		public override string ToString()
+		{
+			return DepotFile;
+		}
+	}
+}
diff --git a/Engine/Source/Programs/DotNETCommon/DotNETUtilities/Perforce/Records/SyncSummaryRecord.cs b/Engine/Source/Programs/DotNETCommon/DotNETUtilities/Perforce/Records/SyncSummaryRecord.cs
new file mode 100644
--- /dev/null
+++ b/Engine/Source/Programs/DotNETCommon/DotNETUtilities/Perforce/Records/SyncSummaryRecord.cs
@@ -0,0 +1,32 @@
+using System;
+using System.Collections.Generic;
+using System.Linq;
+using System.Text;
+using System.Threading.Tasks;
+
+namespace Tools.DotNETCommon.Perforce
+{
+	/// <summary>
+	/// Information about a sync operation
+	/// </summary>
+	public class SyncSummaryRecord
+	{
+		/// <summary>
+		/// The total size of all files synced
+		/// </summary>
+		[PerforceTag("totalFileSize")]
+		public long TotalFileSize;
+
+		/// <summary>
+		/// The total number of files synced
+		/// </summary>
+		[PerforceTag("totalFileCount")]
+		public long TotalFileCount;
+
+		/// <summary>
+		/// The changelist that was synced to
+		/// </summary>
+		[PerforceTag("change")]
+		public int Change;
+	}
+}
diff --git a/Engine/Source/Programs/DotNETCommon/DotNETUtilities/Perforce/Records/UnshelveOptions.cs b/Engine/Source/Programs/DotNETCommon/DotNETUtilities/Perforce/Records/UnshelveOptions.cs
new file mode 100644
--- /dev/null
+++ b/Engine/Source/Programs/DotNETCommon/DotNETUtilities/Perforce/Records/UnshelveOptions.cs
@@ -0,0 +1,30 @@
+using System;
+using System.Collections.Generic;
+using System.Linq;
+using System.Text;
+using System.Threading.Tasks;
+
+namespace Tools.DotNETCommon.Perforce
+{
+	/// <summary>
+	/// Options for the p4 unshelve command
+	/// </summary>
+	[Flags]
+	public enum UnshelveOptions
+	{
+		/// <summary>
+		/// No options specified
+		/// </summary>
+		None = 0,
+
+		/// <summary>
+		/// Force the overwriting of writable (but unopened) files during the unshelve operation.
+		/// </summary>
+		ForceOverwrite = 1,
+
+		/// <summary>
+		/// Preview the results of the unshelve operation without actually restoring the files to your workspace.
+		/// </summary>
+		PreviewOnly = 2,
+	}
+}
diff --git a/Engine/Source/Programs/DotNETCommon/DotNETUtilities/Perforce/Records/UnshelveRecord.cs b/Engine/Source/Programs/DotNETCommon/DotNETUtilities/Perforce/Records/UnshelveRecord.cs
new file mode 100644
--- /dev/null
+++ b/Engine/Source/Programs/DotNETCommon/DotNETUtilities/Perforce/Records/UnshelveRecord.cs
@@ -0,0 +1,41 @@
+using System;
+using System.Collections.Generic;
+using System.Linq;
+using System.Text;
+using System.Threading.Tasks;
+
+namespace Tools.DotNETCommon.Perforce
+{
+	/// <summary>
+	/// Information about an unshelved file
+	/// </summary>
+	public class UnshelveRecord
+	{
+		/// <summary>
+		/// Path to the file in the depot
+		/// </summary>
+		[PerforceTag("depotFile")]
+		public string DepotFile;
+
+		/// <summary>
+		/// The revision number of the file
+		/// </summary>
+		[PerforceTag("rev")]
+		public int Revision;
+
+		/// <summary>
+		/// Open action for the file
+		/// </summary>
+		[PerforceTag("action")]
+		public string Action;
+
+		/// <summary>
+		/// Format this record for display in the debugger
+		/// </summary>
+		/// <returns>Summary of this revision</returns>
+		public override string ToString()
+		{
+			return DepotFile;
+		}
+	}
+}
diff --git a/Engine/Source/Programs/DotNETCommon/DotNETUtilities/Perforce/Records/UpdateChangeOptions.cs b/Engine/Source/Programs/DotNETCommon/DotNETUtilities/Perforce/Records/UpdateChangeOptions.cs
new file mode 100644
--- /dev/null
+++ b/Engine/Source/Programs/DotNETCommon/DotNETUtilities/Perforce/Records/UpdateChangeOptions.cs
@@ -0,0 +1,29 @@
+using System;
+using System.Collections.Generic;
+using System.Linq;
+using System.Text;
+using System.Threading.Tasks;
+
+namespace Tools.DotNETCommon.Perforce
+{
+	/// <summary>
+	/// Options for the UpdateChange command
+	/// </summary>
+	public enum UpdateChangeOptions
+	{
+		/// <summary>
+		/// No options are specified
+		/// </summary>
+		None = 0,
+
+		/// <summary>
+		/// Allows the description, modification date, or user of a submitted changelist to be edited. Editing a submitted changelist requires admin or super access. 
+		/// </summary>
+		Force = 2,
+
+		/// <summary>
+		/// Update a submitted changelist. Only the Jobs:, Description:, or Type: fields can be updated, and only the submitter of the changelist can update the changelist.
+		/// </summary>
+		Submitted = 4,
+	}
+}
diff --git a/Engine/Source/Programs/DotNETCommon/DotNETUtilities/Perforce/Records/UserRecord.cs b/Engine/Source/Programs/DotNETCommon/DotNETUtilities/Perforce/Records/UserRecord.cs
new file mode 100644
--- /dev/null
+++ b/Engine/Source/Programs/DotNETCommon/DotNETUtilities/Perforce/Records/UserRecord.cs
@@ -0,0 +1,71 @@
+using System;
+using System.Collections.Generic;
+using System.Linq;
+using System.Text;
+using System.Threading.Tasks;
+
+namespace Tools.DotNETCommon.Perforce
+{
+	/// <summary>
+	/// Contains information about a Perforce user
+	/// </summary>
+	public class UserRecord
+	{
+		/// <summary>
+		/// The name for the user
+		/// </summary>
+		[PerforceTag("User")]
+		public string UserName;
+
+		/// <summary>
+		/// Registered email address for reviews
+		/// </summary>
+		[PerforceTag("Email")]
+		public string Email;
+
+		/// <summary>
+		/// Last time the user's information was updated
+		/// </summary>
+		[PerforceTag("Update")]
+		public DateTime Update;
+
+		/// <summary>
+		/// Last time the user's information was accessed
+		/// </summary>
+		[PerforceTag("Access")]
+		public DateTime Access;
+
+		/// <summary>
+		/// The user's full name
+		/// </summary>
+		[PerforceTag("FullName")]
+		public string FullName;
+
+		/// <summary>
+		/// Paths which the user is watching
+		/// </summary>
+		[PerforceTag("Reviews", Optional = true)]
+		public List<string> Reviews = new List<string>();
+
+		/// <summary>
+		/// The type of user
+		/// </summary>
+		[PerforceTag("Type")]
+		public string Type;
+
+		/// <summary>
+		/// Method used to authenticate
+		/// </summary>
+		[PerforceTag("AuthMethod")]
+		public string AuthMethod;
+
+		/// <summary>
+		/// Summarize this record for display in the debugger
+		/// </summary>
+		/// <returns>Formatted record</returns>
+		public override string ToString()
+		{
+			return FullName;
+		}
+	}
+}
diff --git a/Engine/Source/Programs/DotNETCommon/DotNETUtilities/Perforce/Records/WhereRecord.cs b/Engine/Source/Programs/DotNETCommon/DotNETUtilities/Perforce/Records/WhereRecord.cs
new file mode 100644
--- /dev/null
+++ b/Engine/Source/Programs/DotNETCommon/DotNETUtilities/Perforce/Records/WhereRecord.cs
@@ -0,0 +1,47 @@
+using System;
+using System.Collections.Generic;
+using System.Linq;
+using System.Text;
+using System.Threading.Tasks;
+
+namespace Tools.DotNETCommon.Perforce
+{
+	/// <summary>
+	/// Contains information about a file's location in the workspace
+	/// </summary>
+	public class WhereRecord
+	{
+		/// <summary>
+		/// Path to the file in the depot
+		/// </summary>
+		[PerforceTag("depotFile")]
+		public string DepotFile;
+
+		/// <summary>
+		/// Path to the file in client syntax
+		/// </summary>
+		[PerforceTag("clientFile")]
+		public string ClientFile;
+
+		/// <summary>
+		/// Path to the file on dist
+		/// </summary>
+		[PerforceTag("path")]
+		public string Path;
+
+		/// <summary>
+		/// Indicates that the given file or path is being unmapped from the workspace
+		/// </summary>
+		[PerforceTag("unmap", Optional = true)]
+		public bool Unmap;
+
+		/// <summary>
+		/// Summarize this record for display in the debugger
+		/// </summary>
+		/// <returns>Formatted record</returns>
+		public override string ToString()
+		{
+			return DepotFile;
+		}
+	}
+}
