UBT: Add extension methods for serializing a nullable type to a BinaryReader/BinaryWriter.

#rnx

diff --git a/Engine/Source/Programs/DotNETCommon/DotNETUtilities/BinaryReaderExtensions.cs b/Engine/Source/Programs/DotNETCommon/DotNETUtilities/BinaryReaderExtensions.cs
--- a/Engine/Source/Programs/DotNETCommon/DotNETUtilities/BinaryReaderExtensions.cs
+++ b/Engine/Source/Programs/DotNETCommon/DotNETUtilities/BinaryReaderExtensions.cs
@@ -17,7 +17,7 @@ namespace Tools.DotNETCommon
 		/// <typeparam name="T">The element type for the array</typeparam>
 		/// <param name="Reader">Reader to read data from</param>
 		/// <returns>Object instance.</returns>
-		public static T ReadObject<T>(this BinaryReader Reader) where T : IBinarySerializable
+		public static T ReadObject<T>(this BinaryReader Reader) where T : class, IBinarySerializable
 		{
 			return (T)Activator.CreateInstance(typeof(T), Reader);
 		}
@@ -38,7 +38,7 @@ namespace Tools.DotNETCommon
 		/// <typeparam name="T">The element type for the array</typeparam>
 		/// <param name="Reader">Reader to read data from</param>
 		/// <returns>Array of objects, as serialized. May be null.</returns>
-		public static T[] ReadArray<T>(this BinaryReader Reader) where T : IBinarySerializable
+		public static T[] ReadArray<T>(this BinaryReader Reader) where T : class, IBinarySerializable
 		{
 			return ReadArray<T>(Reader, () => Reader.ReadObject<T>());
 		}
@@ -83,7 +83,7 @@ namespace Tools.DotNETCommon
 		/// <param name="Reader">Reader to read data from</param>
 		/// <param name="ReadElement">Delegate to call to serialize each element</param>
 		/// <returns>List of objects, as serialized. May be null.</returns>
-		public static List<T> ReadList<T>(this BinaryReader Reader) where T : IBinarySerializable
+		public static List<T> ReadList<T>(this BinaryReader Reader) where T : class, IBinarySerializable
 		{
 			return ReadList<T>(Reader, () => Reader.ReadObject<T>());
 		}
@@ -137,6 +137,25 @@ namespace Tools.DotNETCommon
 			return Items;
 		}
 
+		/// <summary>
+		/// Read a nullable object from a binary reader
+		/// </summary>
+		/// <typeparam name="T">Type of the object</typeparam>
+		/// <param name="Reader">Reader to read data from</param>
+		/// <param name="ReadItem">Function to read the payload, if non-null</param>
+		/// <returns>Object instance or null</returns>
+		public static T ReadNullable<T>(this BinaryReader Reader, Func<T> ReadItem) where T : class
+		{
+			if(Reader.ReadBoolean())
+			{
+				return ReadItem();
+			}
+			else
+			{
+				return null;
+			}
+		}
+
 		/// <summary>
 		/// Reads a value of a specific type from a binary reader
 		/// </summary>
diff --git a/Engine/Source/Programs/DotNETCommon/DotNETUtilities/BinaryWriterExtensions.cs b/Engine/Source/Programs/DotNETCommon/DotNETUtilities/BinaryWriterExtensions.cs
--- a/Engine/Source/Programs/DotNETCommon/DotNETUtilities/BinaryWriterExtensions.cs
+++ b/Engine/Source/Programs/DotNETCommon/DotNETUtilities/BinaryWriterExtensions.cs
@@ -50,9 +50,9 @@ namespace Tools.DotNETCommon
 		/// <typeparam name="T">The array element type</typeparam>
 		/// <param name="Writer">Writer to serialize to</param>
 		/// <param name="Items">Array of items</param>
-		public static void Write<T>(this BinaryWriter Writer, T[] Items) where T : IBinarySerializable
+		public static void Write<T>(this BinaryWriter Writer, T[] Items) where T : class, IBinarySerializable
 		{
-			Write(Writer, Items, Item => Item.Write(Writer));
+			Write(Writer, Items, Item => Writer.Write(Item));
 		}
 
 		/// <summary>
@@ -83,7 +83,7 @@ namespace Tools.DotNETCommon
 		/// </summary>
 		/// <param name="Writer">Writer to serialize to</param>
 		/// <param name="Items">Array of items</param>
-		public static void Write<T>(this BinaryWriter Writer, List<string> Items)
+		public static void Write(this BinaryWriter Writer, List<string> Items)
 		{
 			Write(Writer, Items, Item => Writer.Write(Item));
 		}
@@ -94,9 +94,9 @@ namespace Tools.DotNETCommon
 		/// <typeparam name="T">The array element type</typeparam>
 		/// <param name="Writer">Writer to serialize to</param>
 		/// <param name="Items">Array of items</param>
-		public static void Write<T>(this BinaryWriter Writer, List<T> Items) where T : IBinarySerializable
+		public static void Write<T>(this BinaryWriter Writer, List<T> Items) where T : class, IBinarySerializable
 		{
-			Write(Writer, Items, Item => Item.Write(Writer));
+			Write(Writer, Items, Item => Writer.Write(Item));
 		}
 
 		/// <summary>
@@ -134,7 +134,7 @@ namespace Tools.DotNETCommon
 		/// <returns>Dictionary of objects, as serialized. May be null.</returns>
 		public static void Write<K, V>(this BinaryWriter Writer, Dictionary<K, V> Items, Action<K> WriteKey, Action<V> WriteValue)
 		{
-			if(Items.Count == 0)
+			if(Items == null)
 			{
 				Writer.Write(-1);
 			}
@@ -149,6 +149,26 @@ namespace Tools.DotNETCommon
 			}
 		}
 
+		/// <summary>
+		/// Read a nullable object from a binary reader
+		/// </summary>
+		/// <typeparam name="T">Type of the object</typeparam>
+		/// <param name="Writer">Reader to read data from</param>
+		/// <param name="WriteItem">Function to read the payload, if non-null</param>
+		/// <returns>Object instance or null</returns>
+		public static void WriteNullable<T>(this BinaryWriter Writer, T Item, Action WriteItem) where T : class
+		{
+			if(Item == null)
+			{
+				Writer.Write(false);
+			}
+			else
+			{
+				Writer.Write(true);
+				WriteItem();
+			}
+		}
+
 		/// <summary>
 		/// Writes a value of a specific type to a binary writer
 		/// </summary>
