Missing project file for DotNETUtilities.

#rnx

diff --git a/Engine/Source/Programs/DotNETCommon/DotNETUtilities/DotNETUtilities.csproj b/Engine/Source/Programs/DotNETCommon/DotNETUtilities/DotNETUtilities.csproj
--- a/Engine/Source/Programs/DotNETCommon/DotNETUtilities/DotNETUtilities.csproj
+++ b/Engine/Source/Programs/DotNETCommon/DotNETUtilities/DotNETUtilities.csproj
@@ -77,6 +77,68 @@
     <Compile Include="JsonObject.cs" />
     <Compile Include="Log.cs" />
     <Compile Include="ManagedProcess.cs" />
+    <Compile Include="Perforce\PerforceAttributes.cs" />
+    <Compile Include="Perforce\PerforceChildProcess.cs" />
+    <Compile Include="Perforce\PerforceConnection.cs" />
+    <Compile Include="Perforce\PerforceError.cs" />
+    <Compile Include="Perforce\PerforceException.cs" />
+    <Compile Include="Perforce\PerforceInfo.cs" />
+    <Compile Include="Perforce\PerforceResponse.cs" />
+    <Compile Include="Perforce\PerforceResponseList.cs" />
+    <Compile Include="Perforce\PerforceUtils.cs" />
+    <Compile Include="Perforce\Records\AddOptions.cs" />
+    <Compile Include="Perforce\Records\AddRecord.cs" />
+    <Compile Include="Perforce\Records\ChangeRecord.cs" />
+    <Compile Include="Perforce\Records\ChangesOptions.cs" />
+    <Compile Include="Perforce\Records\ChangesRecord.cs" />
+    <Compile Include="Perforce\Records\ChangeStatus.cs" />
+    <Compile Include="Perforce\Records\ChangeType.cs" />
+    <Compile Include="Perforce\Records\ClientLineEndings.cs" />
+    <Compile Include="Perforce\Records\ClientOptions.cs" />
+    <Compile Include="Perforce\Records\ClientRecord.cs" />
+    <Compile Include="Perforce\Records\ClientsOptions.cs" />
+    <Compile Include="Perforce\Records\ClientsRecord.cs" />
+    <Compile Include="Perforce\Records\ClientSubmitOptions.cs" />
+    <Compile Include="Perforce\Records\DeleteChangeOptions.cs" />
+    <Compile Include="Perforce\Records\DeleteClientOptions.cs" />
+    <Compile Include="Perforce\Records\DescribeFileRecord.cs" />
+    <Compile Include="Perforce\Records\DescribeRecord.cs" />
+    <Compile Include="Perforce\Records\EditOptions.cs" />
+    <Compile Include="Perforce\Records\EditRecord.cs" />
+    <Compile Include="Perforce\Records\FileAction.cs" />
+    <Compile Include="Perforce\Records\FileLogOptions.cs" />
+    <Compile Include="Perforce\Records\FileLogRecord.cs" />
+    <Compile Include="Perforce\Records\FStatOptions.cs" />
+    <Compile Include="Perforce\Records\FStatRecord.cs" />
+    <Compile Include="Perforce\Records\GetChangeOptions.cs" />
+    <Compile Include="Perforce\Records\InfoOptions.cs" />
+    <Compile Include="Perforce\Records\InfoRecord.cs" />
+    <Compile Include="Perforce\Records\IntegrationAction.cs" />
+    <Compile Include="Perforce\Records\IntegrationRecord.cs" />
+    <Compile Include="Perforce\Records\PrintRecord.cs" />
+    <Compile Include="Perforce\Records\ReconcileOptions.cs" />
+    <Compile Include="Perforce\Records\ReconcileRecord.cs" />
+    <Compile Include="Perforce\Records\ResolveOptions.cs" />
+    <Compile Include="Perforce\Records\ResolveRecord.cs" />
+    <Compile Include="Perforce\Records\RevertOptions.cs" />
+    <Compile Include="Perforce\Records\RevertRecord.cs" />
+    <Compile Include="Perforce\Records\RevisionRecord.cs" />
+    <Compile Include="Perforce\Records\ShelveOptions.cs" />
+    <Compile Include="Perforce\Records\ShelveRecord.cs" />
+    <Compile Include="Perforce\Records\StreamOptions.cs" />
+    <Compile Include="Perforce\Records\StreamRecord.cs" />
+    <Compile Include="Perforce\Records\StreamType.cs" />
+    <Compile Include="Perforce\Records\SubmitOptions.cs" />
+    <Compile Include="Perforce\Records\SubmitRecord.cs" />
+    <Compile Include="Perforce\Records\SwitchClientOptions.cs" />
+    <Compile Include="Perforce\Records\SyncOptions.cs" />
+    <Compile Include="Perforce\Records\SyncRecord.cs" />
+    <Compile Include="Perforce\Records\SyncSummaryRecord.cs" />
+    <Compile Include="Perforce\Records\UnshelveOptions.cs" />
+    <Compile Include="Perforce\Records\UnshelveRecord.cs" />
+    <Compile Include="Perforce\Records\UpdateChangeOptions.cs" />
+    <Compile Include="Perforce\Records\UserRecord.cs" />
+    <Compile Include="Perforce\Records\WhereRecord.cs" />
     <Compile Include="Properties\AssemblyInfo.cs" />
     <Compile Include="ResXResource\UEResXReader.cs" />
     <Compile Include="ResXResource\UEResXWriter.cs" />
