UAT: Add support for setting a status message through the log class. Allows writing transient messages (eg. progress messages) which will be cleared out before writing other messages. Best used through the LogStatusScope class, which can set a status message for the duration of a using() block.

As part of this change, the console no longer has to be added as a dedicated trace listener. Since we already special-case this listener when formatting log output, it's easier to just keep the implementation separate to the other trace listeners.

diff --git a/Engine/Source/Programs/AutomationTool/AutomationUtils/ParallelExecutor.cs b/Engine/Source/Programs/AutomationTool/AutomationUtils/ParallelExecutor.cs
--- a/Engine/Source/Programs/AutomationTool/AutomationUtils/ParallelExecutor.cs
+++ b/Engine/Source/Programs/AutomationTool/AutomationUtils/ParallelExecutor.cs
@@ -85,7 +85,7 @@ namespace AutomationTool
 
 			CommandUtils.Log("Building with {0} {1}...", MaxProcesses, (MaxProcesses == 1)? "process" : "processes");
 
-			using (ScopedLogIndent Indent = new ScopedLogIndent("  "))
+			using (LogIndentScope Indent = new LogIndentScope("  "))
 			{
 				// Create the list of things to process
 				List<BuildAction> QueuedActions = new List<BuildAction>();
diff --git a/Engine/Source/Programs/AutomationTool/Program.cs b/Engine/Source/Programs/AutomationTool/Program.cs
--- a/Engine/Source/Programs/AutomationTool/Program.cs
+++ b/Engine/Source/Programs/AutomationTool/Program.cs
@@ -57,7 +57,6 @@ namespace AutomationTool
                 bColorConsoleOutput: true,
                 TraceListeners: new TraceListener[]
                 {
-                    new Tools.DotNETCommon.ConsoleTraceListener(),
 					StartupListener
 				}
 			);
diff --git a/Engine/Source/Programs/DotNETCommon/DotNETUtilities/ConsoleTraceListener.cs b/Engine/Source/Programs/DotNETCommon/DotNETUtilities/ConsoleTraceListener.cs
deleted file mode 100644
diff --git a/Engine/Source/Programs/DotNETCommon/DotNETUtilities/DotNETUtilities.csproj b/Engine/Source/Programs/DotNETCommon/DotNETUtilities/DotNETUtilities.csproj
--- a/Engine/Source/Programs/DotNETCommon/DotNETUtilities/DotNETUtilities.csproj
+++ b/Engine/Source/Programs/DotNETCommon/DotNETUtilities/DotNETUtilities.csproj
@@ -57,7 +57,6 @@
     <Compile Include="AssemblyUtils.cs" />
     <Compile Include="BinaryReaderExtensions.cs" />
     <Compile Include="BinaryWriterExtensions.cs" />
-    <Compile Include="ConsoleTraceListener.cs" />
     <Compile Include="CsProjectInfo.cs" />
     <Compile Include="DirectoryReference.cs" />
     <Compile Include="ExceptionUtils.cs" />
@@ -76,6 +75,8 @@
     <Compile Include="FileSystemReference.cs" />
     <Compile Include="JsonObject.cs" />
     <Compile Include="Log.cs" />
+    <Compile Include="LogIndentScope.cs" />
+    <Compile Include="LogStatusScope.cs" />
     <Compile Include="ManagedProcess.cs" />
     <Compile Include="Perforce\PerforceAttributes.cs" />
     <Compile Include="Perforce\PerforceChildProcess.cs" />
diff --git a/Engine/Source/Programs/DotNETCommon/DotNETUtilities/Log.cs b/Engine/Source/Programs/DotNETCommon/DotNETUtilities/Log.cs
--- a/Engine/Source/Programs/DotNETCommon/DotNETUtilities/Log.cs
+++ b/Engine/Source/Programs/DotNETCommon/DotNETUtilities/Log.cs
@@ -84,12 +84,39 @@ namespace Tools.DotNETCommon
 	/// </summary>
 	static public class Log
 	{
+		/// <summary>
+		/// Temporary status message displayed on the console.
+		/// </summary>
+		[DebuggerDisplay("{HeadingText}")]
+		class StatusMessage
+		{
+			/// <summary>
+			/// The heading for this status message.
+			/// </summary>
+			public string HeadingText;
+
+			/// <summary>
+			/// The current status text.
+			/// </summary>
+			public string CurrentText;
+
+			/// <summary>
+			/// Whether the heading has been written to the console. Before the first time that lines are output to the log in the midst of a status scope, the heading will be written on a line of its own first.
+			/// </summary>
+			public bool bHasFlushedHeadingText;
+		}
+
 		/// <summary>
 		/// Guard our initialization. Mainly used by top level exception handlers to ensure its safe to call a logging function.
 		/// In general user code should not concern itself checking for this.
 		/// </summary>
 		private static bool bIsInitialized = false;
 
+		/// <summary>
+		/// Object used for synchronization
+		/// </summary>
+		private static object SyncObject = new object();
+
 		/// <summary>
 		/// When true, verbose logging is enabled.
 		/// </summary>
@@ -104,6 +131,7 @@ namespace Tools.DotNETCommon
 		/// When true, warnings and errors will have a prefix suitable for display by MSBuild (avoiding error messages showing as (EXEC : Error : ")
 		/// </summary>
 		private static bool bLogProgramNameWithSeverity = false;
+
 		/// <summary>
 		/// When true, logs will have the calling mehod prepended to the output as MethodName:
 		/// </summary>
@@ -134,6 +162,16 @@ namespace Tools.DotNETCommon
 		/// </summary>
 		private static List<string> WriteOnceSet = new List<string>();
 
+		/// <summary>
+		/// Stack of status scope information.
+		/// </summary>
+		private static Stack<StatusMessage> StatusMessageStack = new Stack<StatusMessage>();
+
+		/// <summary>
+		/// The currently visible status text
+		/// </summary>
+		private static string StatusText = "";
+
 		/// <summary>
 		/// Indent added to every output line
 		/// </summary>
@@ -348,46 +386,73 @@ namespace Tools.DotNETCommon
 
 			if (Verbosity <= LogLevel)
 			{
-				// Do console color highlighting here.
-				bool bResetConsoleColor = false;
-				if (bColorConsoleOutput)
+				lock (SyncObject)
 				{
-					if(Verbosity == LogEventType.Warning)
+					// Output to all the other trace listeners
+					List<string> Lines = FormatMessage(StackFramesToSkip + 1, Verbosity, FormatOptions, false, Format, Args);
+					foreach (TraceListener Listener in Trace.Listeners)
 					{
-						Console.ForegroundColor = ConsoleColor.Yellow;
-						bResetConsoleColor = true;
-					}
-					if(Verbosity <= LogEventType.Error)
-					{
-						Console.ForegroundColor = ConsoleColor.Red;
-						bResetConsoleColor = true;
+						foreach(string Line in Lines)
+						{
+							Listener.WriteLine(Line);
+						}
+						Listener.Flush();
 					}
-				}
-				try
-				{
-					lock (((System.Collections.ICollection)Trace.Listeners).SyncRoot)
+
+					// Handle the console output separately; we format things differently
+					if (Verbosity != LogEventType.Log || LogLevel >= LogEventType.Verbose)
 					{
-						foreach (TraceListener l in Trace.Listeners)
+						if(StatusMessageStack.Count > 0)
 						{
-							bool bIsConsole = l is Tools.DotNETCommon.ConsoleTraceListener;
-							if (Verbosity != LogEventType.Log || !bIsConsole || LogLevel >= LogEventType.Verbose)
+							StatusMessage CurrentStatus = StatusMessageStack.Peek();
+							if(CurrentStatus.HeadingText.Length > 0 && !CurrentStatus.bHasFlushedHeadingText)
+							{
+								SetStatusText(CurrentStatus.HeadingText);
+								Console.WriteLine();
+								StatusText = "";
+								CurrentStatus.bHasFlushedHeadingText = true;
+							}
+							else
 							{
-								List<string> Lines = FormatMessage(StackFramesToSkip + 1, Verbosity, FormatOptions, bIsConsole, Format, Args);
-								foreach (string Line in Lines)
-								{
-									l.WriteLine(Line);
-								}
+								SetStatusText("");
 							}
-							l.Flush();
 						}
-					}
-				}
-				finally
-				{
-					// make sure we always put the console color back.
-					if(bResetConsoleColor)
-					{
-						Console.ResetColor();
+
+						bool bResetConsoleColor = false;
+						if (bColorConsoleOutput)
+						{
+							if(Verbosity == LogEventType.Warning)
+							{
+								Console.ForegroundColor = ConsoleColor.Yellow;
+								bResetConsoleColor = true;
+							}
+							if(Verbosity <= LogEventType.Error)
+							{
+								Console.ForegroundColor = ConsoleColor.Red;
+								bResetConsoleColor = true;
+							}
+						}
+						try
+						{
+							List<string> ConsoleLines = FormatMessage(StackFramesToSkip + 1, Verbosity, FormatOptions, true, Format, Args);
+							foreach (string ConsoleLine in ConsoleLines)
+							{
+								Console.WriteLine(ConsoleLine);
+							}
+						}
+						finally
+						{
+							// make sure we always put the console color back.
+							if(bResetConsoleColor)
+							{
+								Console.ResetColor();
+							}
+						}
+
+						if(StatusMessageStack.Count > 0)
+						{
+							SetStatusText(StatusMessageStack.Peek().CurrentText);
+						}
 					}
 				}
 			}
@@ -636,37 +701,98 @@ namespace Tools.DotNETCommon
 		{
 			WriteLinePrivate(1, true, LogEventType.Log, LogFormatOptions.None, Format, Args);
 		}
-	}
 
-	/// <summary>
-	/// Class to apply a log indent for the lifetime of an object 
-	/// </summary>
-	public class ScopedLogIndent : IDisposable
-	{
 		/// <summary>
-		/// The previous indent
+		/// Enter a scope with the given status message. The message will be written to the console without a newline, allowing it to be updated through subsequent calls to UpdateStatus().
+		/// The message will be written to the log immediately. If another line is written while in a status scope, the initial status message is flushed to the console first.
 		/// </summary>
-		string PrevIndent;
+		/// <param name="Message">The status message</param>
+		[Conditional("TRACE")]
+		[MethodImplAttribute(MethodImplOptions.NoInlining)]
+		public static void PushStatus(string Message)
+		{
+			lock(SyncObject)
+			{
+				StatusMessage NewStatusMessage = new StatusMessage();
+				NewStatusMessage.HeadingText = Message;
+				NewStatusMessage.CurrentText = Message;
+				StatusMessageStack.Push(NewStatusMessage);
+
+				if(Message.Length > 0)
+				{
+					WriteLine(LogEventType.Log, "{0}", Message);
+					SetStatusText(Message);
+				}
+			}
+		}
 
 		/// <summary>
-		/// Constructor
+		/// Updates the current status message. This will overwrite the previous status line.
 		/// </summary>
-		/// <param name="Indent">Indent to append to the existing indent</param>
-		public ScopedLogIndent(string Indent)
+		/// <param name="Message">The status message</param>
+		[Conditional("TRACE")]
+		[MethodImplAttribute(MethodImplOptions.NoInlining)]
+		public static void UpdateStatus(string Message)
 		{
-			PrevIndent = Log.Indent;
-			Log.Indent += Indent;
+			lock(SyncObject)
+			{
+				StatusMessage CurrentStatusMessage = StatusMessageStack.Peek();
+				CurrentStatusMessage.CurrentText = Message;
+
+				SetStatusText(Message);
+			}
 		}
 
 		/// <summary>
-		/// Restore the log indent to normal
+		/// Updates the Pops the top status message from the stack. The mess
 		/// </summary>
-		public void Dispose()
+		/// <param name="Message"></param>
+		[Conditional("TRACE")]
+		[MethodImplAttribute(MethodImplOptions.NoInlining)]
+		public static void PopStatus()
 		{
-			if (PrevIndent != null)
+			lock(SyncObject)
 			{
-				Log.Indent = PrevIndent;
-				PrevIndent = null;
+				if(StatusText.Length > 0)
+				{
+					Console.WriteLine();
+					StatusText = "";
+				}
+				StatusMessageStack.Pop();
+			}
+		}
+
+		/// <summary>
+		/// Update the status text. For internal use only; does not modify the StatusMessageStack objects.
+		/// </summary>
+		/// <param name="NewStatusText">New status text to display</param>
+		private static void SetStatusText(string NewStatusText)
+		{
+			if(NewStatusText.Length > 0)
+			{
+				NewStatusText = Indent + NewStatusText;
+			}
+
+			if(StatusText != NewStatusText)
+			{
+				int NumCommonChars = 0;
+				while(NumCommonChars < StatusText.Length && NumCommonChars < NewStatusText.Length && StatusText[NumCommonChars] == NewStatusText[NumCommonChars])
+				{
+					NumCommonChars++;
+				}
+
+				StringBuilder Text = new StringBuilder();
+				Text.Append('\b', StatusText.Length - NumCommonChars);
+				Text.Append(NewStatusText, NumCommonChars, NewStatusText.Length - NumCommonChars);
+				if(NewStatusText.Length < StatusText.Length)
+				{
+					int NumChars = StatusText.Length - NewStatusText.Length;
+					Text.Append(' ', NumChars);
+					Text.Append('\b', NumChars);
+				}
+				Console.Write(Text.ToString());
+
+				StatusText = NewStatusText;
 			}
 		}
 	}
diff --git a/Engine/Source/Programs/DotNETCommon/DotNETUtilities/LogIndentScope.cs b/Engine/Source/Programs/DotNETCommon/DotNETUtilities/LogIndentScope.cs
new file mode 100644
--- /dev/null
+++ b/Engine/Source/Programs/DotNETCommon/DotNETUtilities/LogIndentScope.cs
@@ -0,0 +1,41 @@
+using System;
+using System.Collections.Generic;
+using System.Linq;
+using System.Text;
+using System.Threading.Tasks;
+
+namespace Tools.DotNETCommon
+{
+	/// <summary>
+	/// Class to apply a log indent for the lifetime of an object 
+	/// </summary>
+	public class LogIndentScope : IDisposable
+	{
+		/// <summary>
+		/// The previous indent
+		/// </summary>
+		string PrevIndent;
+
+		/// <summary>
+		/// Constructor
+		/// </summary>
+		/// <param name="Indent">Indent to append to the existing indent</param>
+		public LogIndentScope(string Indent)
+		{
+			PrevIndent = Log.Indent;
+			Log.Indent += Indent;
+		}
+
+		/// <summary>
+		/// Restore the log indent to normal
+		/// </summary>
+		public void Dispose()
+		{
+			if (PrevIndent != null)
+			{
+				Log.Indent = PrevIndent;
+				PrevIndent = null;
+			}
+		}
+	}
+}
diff --git a/Engine/Source/Programs/DotNETCommon/DotNETUtilities/LogStatusScope.cs b/Engine/Source/Programs/DotNETCommon/DotNETUtilities/LogStatusScope.cs
new file mode 100644
--- /dev/null
+++ b/Engine/Source/Programs/DotNETCommon/DotNETUtilities/LogStatusScope.cs
@@ -0,0 +1,96 @@
+using System;
+using System.Collections.Generic;
+using System.Linq;
+using System.Text;
+using System.Threading.Tasks;
+
+namespace Tools.DotNETCommon
+{
+	/// <summary>
+	/// Writes a status message to the log, which can be updated with a progress indicator as a slow task is being performed.
+	/// </summary>
+	public class LogStatusScope : IDisposable
+	{
+		/// <summary>
+		/// The base status message
+		/// </summary>
+		string Message;
+
+		/// <summary>
+		/// Constructor
+		/// </summary>
+		/// <param name="Message">The status message</param>
+		public LogStatusScope(string Message)
+		{
+			this.Message = Message;
+			Log.PushStatus(Message);
+		}
+
+		/// <summary>
+		/// Constructor
+		/// </summary>
+		/// <param name="Format">The format specifier for the message</param>
+		/// <param name="Args">Arguments for the status message</param>
+		public LogStatusScope(string Format, params object[] Args)
+			: this(String.Format(Format, Args))
+		{
+		}
+
+		/// <summary>
+		/// Updates the base status message passed into the constructor.
+		/// </summary>
+		/// <param name="Message">The status message</param>
+		public void SetMessage(string Message)
+		{
+			this.Message = Message;
+			Log.UpdateStatus(Message);
+		}
+
+		/// <summary>
+		/// Updates the base status message passed into the constructor.
+		/// </summary>
+		/// <param name="Format">The format specifier for the message</param>
+		/// <param name="Args">Arguments for the status message</param>
+		public void SetMessage(string Format, params object[] Args)
+		{
+			SetMessage(String.Format(Format, Args));
+		}
+
+		/// <summary>
+		/// Appends a progress string to the status message. Overwrites any previous progress message.
+		/// </summary>
+		/// <param name="Progress">The progress message</param>
+		public void SetProgress(string Progress)
+		{
+			StringBuilder FullMessage = new StringBuilder(Message);
+			FullMessage.Append(' ');
+			FullMessage.Append(Progress);
+			Log.UpdateStatus(FullMessage.ToString());
+		}
+
+		/// <summary>
+		/// Appends a progress string to the status message. Overwrites any previous progress message.
+		/// </summary>
+		/// <param name="Format">The format specifier for the message</param>
+		/// <param name="Args">Arguments for the status message</param>
+		public void SetProgress(string Format, params object[] Args)
+		{
+			StringBuilder FullMessage = new StringBuilder(Message);
+			FullMessage.Append(' ');
+			FullMessage.AppendFormat(Format, Args);
+			Log.UpdateStatus(FullMessage.ToString());
+		}
+
+		/// <summary>
+		/// Pops the status message from the log.
+		/// </summary>
+		public void Dispose()
+		{
+			if(Message != null)
+			{
+				Log.PopStatus();
+				Message = null;
+			}
+		}
+	}
+}
diff --git a/Engine/Source/Programs/UnrealBuildTool/ProjectFiles/ProjectFileGenerator.cs b/Engine/Source/Programs/UnrealBuildTool/ProjectFiles/ProjectFileGenerator.cs
--- a/Engine/Source/Programs/UnrealBuildTool/ProjectFiles/ProjectFileGenerator.cs
+++ b/Engine/Source/Programs/UnrealBuildTool/ProjectFiles/ProjectFileGenerator.cs
@@ -1737,7 +1737,7 @@ namespace UnrealBuildTool
 						}
 						catch(Exception Ex)
 						{
-							Progress.LogMessage(LogEventType.Warning, "Exception while generating include data for {0}: {1}", CurTarget.TargetFilePath.GetFileNameWithoutAnyExtensions(), Ex.ToString());
+							Log.TraceWarning("Exception while generating include data for {0}: {1}", CurTarget.TargetFilePath.GetFileNameWithoutAnyExtensions(), Ex.ToString());
 						}
 
 						ProjectFileGenerator.OnlyGenerateIntelliSenseDataForProject = null;
diff --git a/Engine/Source/Programs/UnrealBuildTool/System/ParallelExecutor.cs b/Engine/Source/Programs/UnrealBuildTool/System/ParallelExecutor.cs
--- a/Engine/Source/Programs/UnrealBuildTool/System/ParallelExecutor.cs
+++ b/Engine/Source/Programs/UnrealBuildTool/System/ParallelExecutor.cs
@@ -136,7 +136,7 @@ namespace UnrealBuildTool
 			}
 
 			// Execute the actions
-			using (ScopedLogIndent Indent = new ScopedLogIndent("  "))
+			using (LogIndentScope Indent = new LogIndentScope("  "))
 			{
 				// Create a job object for all the child processes
 				bool bResult = true;
diff --git a/Engine/Source/Programs/UnrealBuildTool/System/ProgressWriter.cs b/Engine/Source/Programs/UnrealBuildTool/System/ProgressWriter.cs
--- a/Engine/Source/Programs/UnrealBuildTool/System/ProgressWriter.cs
+++ b/Engine/Source/Programs/UnrealBuildTool/System/ProgressWriter.cs
@@ -20,25 +20,31 @@ namespace UnrealBuildTool
 		public static bool bWriteMarkup = false;
 
 		/// <summary>
-		/// Whether to write messages to the console
+		/// The name to include with the status message
 		/// </summary>
-		bool bWriteToConsole;
 		string Message;
-		int NumCharsToBackspaceOver;
+
+		/// <summary>
+		/// The inner scope object
+		/// </summary>
+		LogStatusScope Status;
+
+		/// <summary>
+		/// The current progress message
+		/// </summary>
 		string CurrentProgressString;
 
 		/// <summary>
 		/// Constructor
 		/// </summary>
 		/// <param name="InMessage">The message to display before the progress percentage</param>
-		/// <param name="bInWriteToConsole">Whether to write progress message to the console</param>
-		public ProgressWriter(string InMessage, bool bInWriteToConsole)
+		/// <param name="bInUpdateStatus">Whether to write messages to the console</param>
+		public ProgressWriter(string InMessage, bool bInUpdateStatus)
 		{
 			Message = InMessage;
-			bWriteToConsole = bInWriteToConsole;
-			if (!bWriteMarkup && bWriteToConsole)
+			if(bInUpdateStatus)
 			{
-				Console.Write(Message + " ");
+				Status = new LogStatusScope(InMessage);
 			}
 			Write(0, 100);
 		}
@@ -48,25 +54,10 @@ namespace UnrealBuildTool
 		/// </summary>
 		public void Dispose()
 		{
-			if (!bWriteMarkup && bWriteToConsole)
-			{
-				Console.WriteLine();
-			}
-		}
-
-		/// <summary>
-		/// Write a message to the log, clearing the current progress percentage first
-		/// </summary>
-		public void LogMessage(LogEventType Verbosity, string Format, params object[] Args)
-		{
-			if(bWriteToConsole)
-			{
-				Console.Write(new string('\b', Message.Length + 1 + NumCharsToBackspaceOver) + new string(' ', Message.Length + 1 + NumCharsToBackspaceOver) + new string('\b', Message.Length + 1 + NumCharsToBackspaceOver));
-			}
-			Log.WriteLine(Verbosity, Format, Args);
-			if(bWriteToConsole)
+			if(Status != null)
 			{
-				Console.Write("{0} {1}", Message, CurrentProgressString);
+				Status.Dispose();
+				Status = null;
 			}
 		}
 
@@ -86,17 +77,9 @@ namespace UnrealBuildTool
 				{
 					Log.WriteLine(LogEventType.Console, "@progress '{0}' {1}", Message, ProgressString);
 				}
-				else if (bWriteToConsole)
+				if(Status != null)
 				{
-					// Backspace over previous progress value
-					while (NumCharsToBackspaceOver-- > 0)
-					{
-						Console.Write("\b");
-					}
-
-					// Display updated progress string and keep track of how long it was
-					NumCharsToBackspaceOver = ProgressString.Length;
-					Console.Write(ProgressString);
+					Status.SetProgress(ProgressString);
 				}
 			}
 		}
diff --git a/Engine/Source/Programs/UnrealBuildTool/UnrealBuildTool.cs b/Engine/Source/Programs/UnrealBuildTool/UnrealBuildTool.cs
--- a/Engine/Source/Programs/UnrealBuildTool/UnrealBuildTool.cs
+++ b/Engine/Source/Programs/UnrealBuildTool/UnrealBuildTool.cs
@@ -370,7 +370,6 @@ namespace UnrealBuildTool
                 bColorConsoleOutput: true,
                 TraceListeners: new TraceListener[]
                 {
-                    new Tools.DotNETCommon.ConsoleTraceListener(),
 					StartupListener
 				}
 			);
