WorkspaceTool: Utility to perform fast switches between workspaces using a local cache of files (moving files to/from the cache rather than copying), and allow quickly cleaning a workspace based on timestamps.

diff --git a/Engine/Source/Programs/DotNETCommon/DotNETUtilities/Perforce/PerforceUtils.cs b/Engine/Source/Programs/DotNETCommon/DotNETUtilities/Perforce/PerforceUtils.cs
--- a/Engine/Source/Programs/DotNETCommon/DotNETUtilities/Perforce/PerforceUtils.cs
+++ b/Engine/Source/Programs/DotNETCommon/DotNETUtilities/Perforce/PerforceUtils.cs
@@ -6,7 +6,7 @@ using System.Threading.Tasks;
 
 namespace Tools.DotNETCommon.Perforce
 {
-	class PerforceUtils
+	public static class PerforceUtils
 	{
 		static public string EscapePath(string Path)
 		{
