Prevent doubled-up lines when writing status updates with console log verbosity.

#rnx

diff --git a/Engine/Source/Programs/DotNETCommon/DotNETUtilities/Log.cs b/Engine/Source/Programs/DotNETCommon/DotNETUtilities/Log.cs
--- a/Engine/Source/Programs/DotNETCommon/DotNETUtilities/Log.cs
+++ b/Engine/Source/Programs/DotNETCommon/DotNETUtilities/Log.cs
@@ -68,6 +68,11 @@ namespace Tools.DotNETCommon
 		/// Never write a severity prefix. Useful for pre-formatted messages that need to be in a particular format for, eg. the Visual Studio output window
 		/// </summary>
 		NoSeverityPrefix = 1,
+
+		/// <summary>
+		/// Do not output text to the console
+		/// </summary>
+		NoConsoleOutput = 2,
 	}
 
 	/// <summary>
@@ -400,7 +405,7 @@ namespace Tools.DotNETCommon
 					}
 
 					// Handle the console output separately; we format things differently
-					if (Verbosity != LogEventType.Log || LogLevel >= LogEventType.Verbose)
+					if ((Verbosity != LogEventType.Log || LogLevel >= LogEventType.Verbose) && (FormatOptions & LogFormatOptions.NoConsoleOutput) == 0)
 					{
 						if(StatusMessageStack.Count > 0)
 						{
@@ -720,7 +725,7 @@ namespace Tools.DotNETCommon
 
 				if(Message.Length > 0)
 				{
-					WriteLine(LogEventType.Log, "{0}", Message);
+					WriteLine(LogEventType.Log, LogFormatOptions.NoConsoleOutput, "{0}", Message);
 					SetStatusText(Message);
 				}
 			}
