UAT: Detect when stdout is redirected and prevent using backspace characters to move the cursor.

diff --git a/Engine/Source/Programs/DotNETCommon/DotNETUtilities/Log.cs b/Engine/Source/Programs/DotNETCommon/DotNETUtilities/Log.cs
--- a/Engine/Source/Programs/DotNETCommon/DotNETUtilities/Log.cs
+++ b/Engine/Source/Programs/DotNETCommon/DotNETUtilities/Log.cs
@@ -151,7 +151,12 @@ namespace Tools.DotNETCommon
 		/// When true, will detect warnings and errors and set the console output color to yellow and red.
 		/// </summary>
 		private static bool bColorConsoleOutput = false;
-		
+
+		/// <summary>
+		/// Whether console output is redirected. This prevents writing status updates that rely on moving the cursor.
+		/// </summary>
+		private static bool bAllowStatusUpdates = true;
+
 		/// <summary>
 		/// When configured, this tracks time since initialization to prepend a timestamp to each log.
 		/// </summary>
@@ -177,6 +182,11 @@ namespace Tools.DotNETCommon
 		/// </summary>
 		private static string StatusText = "";
 
+		/// <summary>
+		/// Last time a status message was pushed to the stack
+		/// </summary>
+		private static Stopwatch StatusTimer = new Stopwatch();
+
 		/// <summary>
 		/// Indent added to every output line
 		/// </summary>
@@ -244,6 +254,7 @@ namespace Tools.DotNETCommon
 			Log.bLogSources = bLogSources;
 			Log.bLogSourcesToConsole = bLogSourcesToConsole;
 			Log.bColorConsoleOutput = bColorConsoleOutput;
+			Log.bAllowStatusUpdates = !Console.IsOutputRedirected;
 
 			// ensure that if InitLogging is called more than once we don't stack listeners.
 			// but always leave the default listener around.
@@ -407,7 +418,7 @@ namespace Tools.DotNETCommon
 					// Handle the console output separately; we format things differently
 					if ((Verbosity != LogEventType.Log || LogLevel >= LogEventType.Verbose) && (FormatOptions & LogFormatOptions.NoConsoleOutput) == 0)
 					{
-						if(StatusMessageStack.Count > 0)
+						if(StatusMessageStack.Count > 0 && bAllowStatusUpdates)
 						{
 							StatusMessage CurrentStatus = StatusMessageStack.Peek();
 							if(CurrentStatus.HeadingText.Length > 0 && !CurrentStatus.bHasFlushedHeadingText)
@@ -454,7 +465,7 @@ namespace Tools.DotNETCommon
 							}
 						}
 
-						if(StatusMessageStack.Count > 0)
+						if(StatusMessageStack.Count > 0 && bAllowStatusUpdates)
 						{
 							SetStatusText(StatusMessageStack.Peek().CurrentText);
 						}
@@ -723,10 +734,19 @@ namespace Tools.DotNETCommon
 				NewStatusMessage.CurrentText = Message;
 				StatusMessageStack.Push(NewStatusMessage);
 
+				StatusTimer.Restart();
+
 				if(Message.Length > 0)
 				{
 					WriteLine(LogEventType.Log, LogFormatOptions.NoConsoleOutput, "{0}", Message);
-					SetStatusText(Message);
+					if(bAllowStatusUpdates)
+					{
+						SetStatusText(Message);
+					}
+					else
+					{
+						Console.WriteLine(Indent + Message);
+					}
 				}
 			}
 		}
@@ -760,7 +780,14 @@ namespace Tools.DotNETCommon
 			{
 				if(StatusText.Length > 0)
 				{
-					Console.WriteLine();
+					if(bAllowStatusUpdates)
+					{
+						Console.WriteLine();
+					}
+					else
+					{
+						Console.WriteLine(StatusText);
+					}
 					StatusText = "";
 				}
 				StatusMessageStack.Pop();
@@ -778,24 +805,37 @@ namespace Tools.DotNETCommon
 				NewStatusText = Indent + NewStatusText;
 			}
 
-			if(StatusText != NewStatusText)
+			if(bAllowStatusUpdates)
 			{
-				int NumCommonChars = 0;
-				while(NumCommonChars < StatusText.Length && NumCommonChars < NewStatusText.Length && StatusText[NumCommonChars] == NewStatusText[NumCommonChars])
+				if(StatusText != NewStatusText)
 				{
-					NumCommonChars++;
-				}
+					int NumCommonChars = 0;
+					while(NumCommonChars < StatusText.Length && NumCommonChars < NewStatusText.Length && StatusText[NumCommonChars] == NewStatusText[NumCommonChars])
+					{
+						NumCommonChars++;
+					}
 
-				StringBuilder Text = new StringBuilder();
-				Text.Append('\b', StatusText.Length - NumCommonChars);
-				Text.Append(NewStatusText, NumCommonChars, NewStatusText.Length - NumCommonChars);
-				if(NewStatusText.Length < StatusText.Length)
+					StringBuilder Text = new StringBuilder();
+					Text.Append('\b', StatusText.Length - NumCommonChars);
+					Text.Append(NewStatusText, NumCommonChars, NewStatusText.Length - NumCommonChars);
+					if(NewStatusText.Length < StatusText.Length)
+					{
+						int NumChars = StatusText.Length - NewStatusText.Length;
+						Text.Append(' ', NumChars);
+						Text.Append('\b', NumChars);
+					}
+					Console.Write(Text.ToString());
+
+					StatusText = NewStatusText;
+				}
+			}
+			else
+			{
+				if(NewStatusText.Length > 0 && StatusTimer.Elapsed.TotalSeconds > 10.0)
 				{
-					int NumChars = StatusText.Length - NewStatusText.Length;
-					Text.Append(' ', NumChars);
-					Text.Append('\b', NumChars);
+					Console.WriteLine(NewStatusText);
+					StatusTimer.Restart();
 				}
-				Console.Write(Text.ToString());
 
 				StatusText = NewStatusText;
 			}
