WorkspaceTool: Reduce unnecessary logging of status messages when console output is not redirected.

#rnx

diff --git a/Engine/Source/Programs/DotNETCommon/DotNETUtilities/Log.cs b/Engine/Source/Programs/DotNETCommon/DotNETUtilities/Log.cs
--- a/Engine/Source/Programs/DotNETCommon/DotNETUtilities/Log.cs
+++ b/Engine/Source/Programs/DotNETCommon/DotNETUtilities/Log.cs
@@ -418,10 +418,10 @@ namespace Tools.DotNETCommon
 					// Handle the console output separately; we format things differently
 					if ((Verbosity != LogEventType.Log || LogLevel >= LogEventType.Verbose) && (FormatOptions & LogFormatOptions.NoConsoleOutput) == 0)
 					{
-						if(StatusMessageStack.Count > 0 && bAllowStatusUpdates)
+						if(StatusMessageStack.Count > 0)
 						{
 							StatusMessage CurrentStatus = StatusMessageStack.Peek();
-							if(CurrentStatus.HeadingText.Length > 0 && !CurrentStatus.bHasFlushedHeadingText)
+							if(CurrentStatus.HeadingText.Length > 0 && !CurrentStatus.bHasFlushedHeadingText && bAllowStatusUpdates)
 							{
 								SetStatusText(CurrentStatus.HeadingText);
 								Console.WriteLine();
@@ -739,14 +739,7 @@ namespace Tools.DotNETCommon
 				if(Message.Length > 0)
 				{
 					WriteLine(LogEventType.Log, LogFormatOptions.NoConsoleOutput, "{0}", Message);
-					if(bAllowStatusUpdates)
-					{
-						SetStatusText(Message);
-					}
-					else
-					{
-						Console.WriteLine(Indent + Message);
-					}
+					SetStatusText(Message);
 				}
 			}
 		}
@@ -764,7 +757,11 @@ namespace Tools.DotNETCommon
 				StatusMessage CurrentStatusMessage = StatusMessageStack.Peek();
 				CurrentStatusMessage.CurrentText = Message;
 
-				SetStatusText(Message);
+				if(bAllowStatusUpdates || StatusTimer.Elapsed.TotalSeconds > 10.0)
+				{
+					SetStatusText(Message);
+					StatusTimer.Restart();
+				}
 			}
 		}
 
@@ -778,18 +775,15 @@ namespace Tools.DotNETCommon
 		{
 			lock(SyncObject)
 			{
+				StatusMessage CurrentStatusMessage = StatusMessageStack.Peek();
+				SetStatusText(CurrentStatusMessage.CurrentText);
+
 				if(StatusText.Length > 0)
 				{
-					if(bAllowStatusUpdates)
-					{
-						Console.WriteLine();
-					}
-					else
-					{
-						Console.WriteLine(StatusText);
-					}
+					Console.WriteLine();
 					StatusText = "";
 				}
+
 				StatusMessageStack.Pop();
 			}
 		}
@@ -805,39 +799,35 @@ namespace Tools.DotNETCommon
 				NewStatusText = Indent + NewStatusText;
 			}
 
-			if(bAllowStatusUpdates)
+			if(StatusText != NewStatusText)
 			{
-				if(StatusText != NewStatusText)
+				int NumCommonChars = 0;
+				while(NumCommonChars < StatusText.Length && NumCommonChars < NewStatusText.Length && StatusText[NumCommonChars] == NewStatusText[NumCommonChars])
 				{
-					int NumCommonChars = 0;
-					while(NumCommonChars < StatusText.Length && NumCommonChars < NewStatusText.Length && StatusText[NumCommonChars] == NewStatusText[NumCommonChars])
-					{
-						NumCommonChars++;
-					}
-
-					StringBuilder Text = new StringBuilder();
-					Text.Append('\b', StatusText.Length - NumCommonChars);
-					Text.Append(NewStatusText, NumCommonChars, NewStatusText.Length - NumCommonChars);
-					if(NewStatusText.Length < StatusText.Length)
-					{
-						int NumChars = StatusText.Length - NewStatusText.Length;
-						Text.Append(' ', NumChars);
-						Text.Append('\b', NumChars);
-					}
-					Console.Write(Text.ToString());
+					NumCommonChars++;
+				}
 
-					StatusText = NewStatusText;
+				if(!bAllowStatusUpdates && NumCommonChars < StatusText.Length)
+				{
+					// Prevent writing backspace characters if the console doesn't support it
+					Console.WriteLine();
+					StatusText = "";
+					NumCommonChars = 0;
 				}
-			}
-			else
-			{
-				if(NewStatusText.Length > 0 && StatusTimer.Elapsed.TotalSeconds > 10.0)
+
+				StringBuilder Text = new StringBuilder();
+				Text.Append('\b', StatusText.Length - NumCommonChars);
+				Text.Append(NewStatusText, NumCommonChars, NewStatusText.Length - NumCommonChars);
+				if(NewStatusText.Length < StatusText.Length)
 				{
-					Console.WriteLine(NewStatusText);
-					StatusTimer.Restart();
+					int NumChars = StatusText.Length - NewStatusText.Length;
+					Text.Append(' ', NumChars);
+					Text.Append('\b', NumChars);
 				}
+				Console.Write(Text.ToString());
 
 				StatusText = NewStatusText;
+				StatusTimer.Restart();
 			}
 		}
 	}
