WorkspaceTool: Improve resilience when long-running actions are interrupted.

* Before moving a large number of files, the workspace state is saved with information about files in the source and target locations, allowing recovery of files in either location.
* The workspace state can be tagged as dirty, meaning that some cache or workspace items may not be correct. If a workspace is loaded with this flag, we check integrity of all files in each category before continuing.
* Gathering files to add and remove is done in a two separate operations.
* Workspace state is captured after every 256mb is synced, to avoid large resyncs if an error occurs.
* If any attempt to move a file from the cache, or copy a file within the workspace fails, we fall back to re-syncing the item from Perforce.

Also other minor improvements:

* Always write out log file, and keep old log files for 3 days.
* Added dump command, to list contents of workspace and cache.

diff --git a/Engine/Source/Programs/DotNETCommon/DotNETUtilities/Log.cs b/Engine/Source/Programs/DotNETCommon/DotNETUtilities/Log.cs
--- a/Engine/Source/Programs/DotNETCommon/DotNETUtilities/Log.cs
+++ b/Engine/Source/Programs/DotNETCommon/DotNETUtilities/Log.cs
@@ -418,21 +418,7 @@ namespace Tools.DotNETCommon
 					// Handle the console output separately; we format things differently
 					if ((Verbosity != LogEventType.Log || LogLevel >= LogEventType.Verbose) && (FormatOptions & LogFormatOptions.NoConsoleOutput) == 0)
 					{
-						if(StatusMessageStack.Count > 0)
-						{
-							StatusMessage CurrentStatus = StatusMessageStack.Peek();
-							if(CurrentStatus.HeadingText.Length > 0 && !CurrentStatus.bHasFlushedHeadingText && bAllowStatusUpdates)
-							{
-								SetStatusText(CurrentStatus.HeadingText);
-								Console.WriteLine();
-								StatusText = "";
-								CurrentStatus.bHasFlushedHeadingText = true;
-							}
-							else
-							{
-								SetStatusText("");
-							}
-						}
+						FlushStatusHeading();
 
 						bool bResetConsoleColor = false;
 						if (bColorConsoleOutput)
@@ -718,6 +704,28 @@ namespace Tools.DotNETCommon
 			WriteLinePrivate(1, true, LogEventType.Log, LogFormatOptions.None, Format, Args);
 		}
 
+		/// <summary>
+		/// Flushes the current status text before writing out a new log line or status message
+		/// </summary>
+		static void FlushStatusHeading()
+		{
+			if(StatusMessageStack.Count > 0)
+			{
+				StatusMessage CurrentStatus = StatusMessageStack.Peek();
+				if(CurrentStatus.HeadingText.Length > 0 && !CurrentStatus.bHasFlushedHeadingText && bAllowStatusUpdates)
+				{
+					SetStatusText(CurrentStatus.HeadingText);
+					Console.WriteLine();
+					StatusText = "";
+					CurrentStatus.bHasFlushedHeadingText = true;
+				}
+				else
+				{
+					SetStatusText("");
+				}
+			}
+		}
+
 		/// <summary>
 		/// Enter a scope with the given status message. The message will be written to the console without a newline, allowing it to be updated through subsequent calls to UpdateStatus().
 		/// The message will be written to the log immediately. If another line is written while in a status scope, the initial status message is flushed to the console first.
@@ -729,6 +737,8 @@ namespace Tools.DotNETCommon
 		{
 			lock(SyncObject)
 			{
+				FlushStatusHeading();
+
 				StatusMessage NewStatusMessage = new StatusMessage();
 				NewStatusMessage.HeadingText = Message;
 				NewStatusMessage.CurrentText = Message;
