WorkspaceTool: Use ManagedProcess to execute Perforce, and add a default implementation of it for non-Win32 platforms.

#rnx

diff --git a/Engine/Source/Programs/AutomationTool/AutomationUtils/ParallelExecutor.cs b/Engine/Source/Programs/AutomationTool/AutomationUtils/ParallelExecutor.cs
--- a/Engine/Source/Programs/AutomationTool/AutomationUtils/ParallelExecutor.cs
+++ b/Engine/Source/Programs/AutomationTool/AutomationUtils/ParallelExecutor.cs
@@ -61,7 +61,7 @@ namespace AutomationTool
 					LogLines.Add(Action.OutputPrefix);
 				}
 
-				using(ManagedProcess Process = new ManagedProcess(ProcessGroup, Action.ToolPath, Action.ToolArguments, Action.WorkingDirectory, Action.Environment, null, ManagedProcessPriority.BelowNormal))
+				using(ManagedProcess Process = new ManagedProcess(ProcessGroup, Action.ToolPath, Action.ToolArguments, Action.WorkingDirectory, Action.Environment, null, ProcessPriorityClass.BelowNormal))
 				{
 					LogLines.AddRange(Process.ReadAllLines());
 					ExitCode = Process.ExitCode;
diff --git a/Engine/Source/Programs/DotNETCommon/DotNETUtilities/ManagedProcess.cs b/Engine/Source/Programs/DotNETCommon/DotNETUtilities/ManagedProcess.cs
--- a/Engine/Source/Programs/DotNETCommon/DotNETUtilities/ManagedProcess.cs
+++ b/Engine/Source/Programs/DotNETCommon/DotNETUtilities/ManagedProcess.cs
@@ -4,6 +4,7 @@ using Microsoft.Win32.SafeHandles;
 using System;
 using System.Collections.Generic;
 using System.ComponentModel;
+using System.Diagnostics;
 using System.IO;
 using System.Linq;
 using System.Runtime.InteropServices;
@@ -13,16 +14,6 @@ using System.Threading.Tasks;
 
 namespace Tools.DotNETCommon
 {
-	/// <summary>
-	/// Possible priority values for new processes.
-	/// </summary>
-	public enum ManagedProcessPriority
-	{
-		BelowNormal,
-		Normal,
-		AboveNormal,
-	}
-
 	/// <summary>
 	/// Tracks a set of processes, and destroys them when the object is disposed.
 	/// </summary>
@@ -79,38 +70,53 @@ namespace Tools.DotNETCommon
 		/// Handle to the native job object that this process is added to. This handle is closed by the Dispose() method (and will automatically be closed by the OS on process exit),
 		/// resulting in the child process being killed.
 		/// </summary>
-		public SafeFileHandle JobHandle
+		internal SafeFileHandle JobHandle
 		{
 			get;
 			private set;
 		}
 
+		/// <summary>
+		/// Determines support for using job objects
+		/// </summary>
+		static internal bool SupportsJobObjects
+		{
+			get { return false; }// Environment.OSVersion.Platform == PlatformID.Win32NT; }
+		}
+
 		/// <summary>
 		/// Constructor
 		/// </summary>
 		public ManagedProcessGroup()
 		{
-			// Create the job object that the child process will be added to
-			JobHandle = CreateJobObject(IntPtr.Zero, IntPtr.Zero);
-			if(JobHandle == null)
+			if(SupportsJobObjects)
 			{
-				throw new Win32Exception();
-			}
+				// Create the job object that the child process will be added to
+				JobHandle = CreateJobObject(IntPtr.Zero, IntPtr.Zero);
+				if(JobHandle == null)
+				{
+					throw new Win32Exception();
+				}
 
-			// Configure the job object to terminate the processes added to it when the handle is closed
-			JOBOBJECT_EXTENDED_LIMIT_INFORMATION LimitInformation = new JOBOBJECT_EXTENDED_LIMIT_INFORMATION();
-			LimitInformation.BasicLimitInformation.LimitFlags = JOB_OBJECT_LIMIT_KILL_ON_JOB_CLOSE | JOB_OBJECT_LIMIT_BREAKAWAY_OK;
+				// Configure the job object to terminate the processes added to it when the handle is closed
+				JOBOBJECT_EXTENDED_LIMIT_INFORMATION LimitInformation = new JOBOBJECT_EXTENDED_LIMIT_INFORMATION();
+				LimitInformation.BasicLimitInformation.LimitFlags = JOB_OBJECT_LIMIT_KILL_ON_JOB_CLOSE | JOB_OBJECT_LIMIT_BREAKAWAY_OK;
 
-			int Length = Marshal.SizeOf(typeof(JOBOBJECT_EXTENDED_LIMIT_INFORMATION));
-			IntPtr LimitInformationPtr = Marshal.AllocHGlobal(Length);
-			Marshal.StructureToPtr(LimitInformation, LimitInformationPtr, false);
+				int Length = Marshal.SizeOf(typeof(JOBOBJECT_EXTENDED_LIMIT_INFORMATION));
+				IntPtr LimitInformationPtr = Marshal.AllocHGlobal(Length);
+				Marshal.StructureToPtr(LimitInformation, LimitInformationPtr, false);
 
-			if(SetInformationJobObject(JobHandle, JobObjectExtendedLimitInformation, LimitInformationPtr, Length) == 0)
-			{
-				throw new Win32Exception();
+				if(SetInformationJobObject(JobHandle, JobObjectExtendedLimitInformation, LimitInformationPtr, Length) == 0)
+				{
+					throw new Win32Exception();
+				}
 			}
 		}
 
+		public void Add(IntPtr ProcessHandle)
+		{
+		}
+
 		/// <summary>
 		/// Dispose of the process group
 		/// </summary>
@@ -192,8 +198,11 @@ namespace Tools.DotNETCommon
 		{
 			CREATE_NO_WINDOW = 0x08000000,
 			CREATE_SUSPENDED = 0x00000004,
-			BELOW_NORMAL_PRIORITY_CLASS = 0x00004000,
 			NORMAL_PRIORITY_CLASS = 0x00000020,
+			IDLE_PRIORITY_CLASS = 0x00000040,
+			HIGH_PRIORITY_CLASS = 0x00000080,
+			REALTIME_PRIORITY_CLASS = 0x00000100,
+			BELOW_NORMAL_PRIORITY_CLASS = 0x00004000,
 			ABOVE_NORMAL_PRIORITY_CLASS = 0x00008000,
 		}
 
@@ -258,6 +267,11 @@ namespace Tools.DotNETCommon
 		/// </summary>
 		StreamReader ReadStream;
 
+		/// <summary>
+		/// Standard process implementation for non-Windows platforms.
+		/// </summary>
+		Process FrameworkProcess;
+
 		/// <summary>
 		/// Static lock object. This is used to synchronize the creation of child processes - in particular, the inheritance of stdout/stderr write pipes. If processes
 		/// inherit pipes meant for other processes, they won't be closed until both terminate.
@@ -273,167 +287,209 @@ namespace Tools.DotNETCommon
 		/// <param name="WorkingDirectory">Working directory for the new process. May be null to use the current working directory.</param>
 		/// <param name="Environment">Environment variables for the new process. May be null, in which case the current process' environment is inherited</param>
 		/// <param name="Input">Text to be passed via stdin to the new process. May be null.</param>
-		public ManagedProcess(ManagedProcessGroup Group, string FileName, string CommandLine, string WorkingDirectory, IReadOnlyDictionary<string, string> Environment, string Input, ManagedProcessPriority Priority)
+		public ManagedProcess(ManagedProcessGroup Group, string FileName, string CommandLine, string WorkingDirectory, IReadOnlyDictionary<string, string> Environment, byte[] Input, ProcessPriorityClass Priority)
 		{
-			// Create the child process
-			IntPtr EnvironmentBlock = IntPtr.Zero;
-			try
+			if(ManagedProcessGroup.SupportsJobObjects)
 			{
-				// Create the environment block for the child process, if necessary.
-				if(Environment != null)
+				// Create the child process
+				IntPtr EnvironmentBlock = IntPtr.Zero;
+				try
 				{
-					// The native format for the environment block is a sequence of null terminated strings with a final null terminator.
-					List<byte> EnvironmentBytes = new List<byte>();
-					foreach(KeyValuePair<string, string> Pair in Environment)
+					// Create the environment block for the child process, if necessary.
+					if(Environment != null)
 					{
-						EnvironmentBytes.AddRange(Encoding.UTF8.GetBytes(Pair.Key));
-						EnvironmentBytes.Add((byte)'=');
-						EnvironmentBytes.AddRange(Encoding.UTF8.GetBytes(Pair.Value));
+						// The native format for the environment block is a sequence of null terminated strings with a final null terminator.
+						List<byte> EnvironmentBytes = new List<byte>();
+						foreach(KeyValuePair<string, string> Pair in Environment)
+						{
+							EnvironmentBytes.AddRange(Encoding.UTF8.GetBytes(Pair.Key));
+							EnvironmentBytes.Add((byte)'=');
+							EnvironmentBytes.AddRange(Encoding.UTF8.GetBytes(Pair.Value));
+							EnvironmentBytes.Add((byte)0);
+						}
 						EnvironmentBytes.Add((byte)0);
-					}
-					EnvironmentBytes.Add((byte)0);
-
-					// Allocate an unmanaged block of memory to store it.
-					EnvironmentBlock = Marshal.AllocHGlobal(EnvironmentBytes.Count);
-					Marshal.Copy(EnvironmentBytes.ToArray(), 0, EnvironmentBlock, EnvironmentBytes.Count);
-				}
 
-				PROCESS_INFORMATION ProcessInfo = new PROCESS_INFORMATION();
-				try
-				{
-					// Get the flags to create the new process
-					ProcessCreationFlags Flags = ProcessCreationFlags.CREATE_NO_WINDOW | ProcessCreationFlags.CREATE_SUSPENDED;
-					switch(Priority)
-					{
-						case ManagedProcessPriority.BelowNormal:
-							Flags |= ProcessCreationFlags.BELOW_NORMAL_PRIORITY_CLASS;
-							break;
-						case ManagedProcessPriority.Normal:
-							Flags |= ProcessCreationFlags.NORMAL_PRIORITY_CLASS;
-							break;
-						case ManagedProcessPriority.AboveNormal:
-							Flags |= ProcessCreationFlags.ABOVE_NORMAL_PRIORITY_CLASS;
-							break;
+						// Allocate an unmanaged block of memory to store it.
+						EnvironmentBlock = Marshal.AllocHGlobal(EnvironmentBytes.Count);
+						Marshal.Copy(EnvironmentBytes.ToArray(), 0, EnvironmentBlock, EnvironmentBytes.Count);
 					}
 
-					// Acquire a global lock before creating inheritable handles. If multiple threads create inheritable handles at the same time, and child processes will inherit them all.
-					// Since we need to wait for output pipes to be closed (in order to consume all output), this can result in output reads not returning until all processes with the same
-					// inherited handles are closed.
-					lock(LockObject)
+					PROCESS_INFORMATION ProcessInfo = new PROCESS_INFORMATION();
+					try
 					{
-						SafeFileHandle StdInRead = null;
-						SafeFileHandle StdOutWrite = null;
-						SafeWaitHandle StdErrWrite = null;
-						try
+						// Get the flags to create the new process
+						ProcessCreationFlags Flags = ProcessCreationFlags.CREATE_NO_WINDOW | ProcessCreationFlags.CREATE_SUSPENDED;
+						switch(Priority)
 						{
-							// Create stdin and stdout pipes for the child process. We'll close the handles for the child process' ends after it's been created.
-							SECURITY_ATTRIBUTES SecurityAttributes = new SECURITY_ATTRIBUTES();
-							SecurityAttributes.nLength = Marshal.SizeOf(SecurityAttributes);
-							SecurityAttributes.bInheritHandle = 1;
+							case ProcessPriorityClass.Normal:
+								Flags |= ProcessCreationFlags.NORMAL_PRIORITY_CLASS;
+								break;
+							case ProcessPriorityClass.Idle:
+								Flags |= ProcessCreationFlags.IDLE_PRIORITY_CLASS;
+								break;
+							case ProcessPriorityClass.High:
+								Flags |= ProcessCreationFlags.HIGH_PRIORITY_CLASS;
+								break;
+							case ProcessPriorityClass.RealTime:
+								Flags |= ProcessCreationFlags.REALTIME_PRIORITY_CLASS;
+								break;
+							case ProcessPriorityClass.BelowNormal:
+								Flags |= ProcessCreationFlags.BELOW_NORMAL_PRIORITY_CLASS;
+								break;
+							case ProcessPriorityClass.AboveNormal:
+								Flags |= ProcessCreationFlags.ABOVE_NORMAL_PRIORITY_CLASS;
+								break;
+						}
 
-							if(CreatePipe(out StdInRead, out StdInWrite, SecurityAttributes, 0) == 0 || SetHandleInformation(StdInWrite, HANDLE_FLAG_INHERIT, 0) == 0)
-							{
-								throw new Win32Exception();
-							}
-							if(CreatePipe(out StdOutRead, out StdOutWrite, SecurityAttributes, 0) == 0 || SetHandleInformation(StdOutRead, HANDLE_FLAG_INHERIT, 0) == 0)
+						// Acquire a global lock before creating inheritable handles. If multiple threads create inheritable handles at the same time, and child processes will inherit them all.
+						// Since we need to wait for output pipes to be closed (in order to consume all output), this can result in output reads not returning until all processes with the same
+						// inherited handles are closed.
+						lock(LockObject)
+						{
+							SafeFileHandle StdInRead = null;
+							SafeFileHandle StdOutWrite = null;
+							SafeWaitHandle StdErrWrite = null;
+							try
 							{
-								throw new Win32Exception();
+								// Create stdin and stdout pipes for the child process. We'll close the handles for the child process' ends after it's been created.
+								SECURITY_ATTRIBUTES SecurityAttributes = new SECURITY_ATTRIBUTES();
+								SecurityAttributes.nLength = Marshal.SizeOf(SecurityAttributes);
+								SecurityAttributes.bInheritHandle = 1;
+
+								if(CreatePipe(out StdInRead, out StdInWrite, SecurityAttributes, 0) == 0 || SetHandleInformation(StdInWrite, HANDLE_FLAG_INHERIT, 0) == 0)
+								{
+									throw new Win32Exception();
+								}
+								if(CreatePipe(out StdOutRead, out StdOutWrite, SecurityAttributes, 0) == 0 || SetHandleInformation(StdOutRead, HANDLE_FLAG_INHERIT, 0) == 0)
+								{
+									throw new Win32Exception();
+								}
+								if(DuplicateHandle(GetCurrentProcess(), StdOutWrite, GetCurrentProcess(), out StdErrWrite, 0, true, DUPLICATE_SAME_ACCESS) == 0)
+								{
+									throw new Win32Exception();
+								}
+
+								// Create the new process as suspended, so we can modify it before it starts executing (and potentially preempting us)
+								STARTUPINFO StartupInfo = new STARTUPINFO();
+								StartupInfo.cb = Marshal.SizeOf(StartupInfo);
+								StartupInfo.hStdInput = StdInRead;
+								StartupInfo.hStdOutput = StdOutWrite;
+								StartupInfo.hStdError = StdErrWrite;
+								StartupInfo.dwFlags = STARTF_USESTDHANDLES;
+
+								if (CreateProcess(null, new StringBuilder("\"" + FileName + "\" " + CommandLine), IntPtr.Zero, IntPtr.Zero, true, Flags, EnvironmentBlock, WorkingDirectory, StartupInfo, ProcessInfo) == 0)
+								{
+									throw new Win32Exception();
+								}
 							}
-							if(DuplicateHandle(GetCurrentProcess(), StdOutWrite, GetCurrentProcess(), out StdErrWrite, 0, true, DUPLICATE_SAME_ACCESS) == 0)
+							finally
 							{
-								throw new Win32Exception();
+								// Close the write ends of the handle. We don't want any other process to be able to inherit these.
+								if(StdInRead != null)
+								{
+									StdInRead.Dispose();
+									StdInRead = null;
+								}
+								if(StdOutWrite != null)
+								{
+									StdOutWrite.Dispose();
+									StdOutWrite = null;
+								}
+								if(StdErrWrite != null)
+								{
+									StdErrWrite.Dispose();
+									StdErrWrite = null;
+								}
 							}
+						}
+
+						// Add it to our job object
+						if(AssignProcessToJobObject(Group.JobHandle, ProcessInfo.hProcess) == 0)
+						{
+							// Support for nested job objects was only addeed in Windows 8; prior to that, assigning processes to job objects would fail. Figure out if we're already in a job, and ignore the error if we are.
+							int OriginalError = Marshal.GetLastWin32Error();
 
-							// Create the new process as suspended, so we can modify it before it starts executing (and potentially preempting us)
-							STARTUPINFO StartupInfo = new STARTUPINFO();
-							StartupInfo.cb = Marshal.SizeOf(StartupInfo);
-							StartupInfo.hStdInput = StdInRead;
-							StartupInfo.hStdOutput = StdOutWrite;
-							StartupInfo.hStdError = StdErrWrite;
-							StartupInfo.dwFlags = STARTF_USESTDHANDLES;
+							bool bProcessInJob;
+							IsProcessInJob(GetCurrentProcess(), IntPtr.Zero, out bProcessInJob);
 
-							if (CreateProcess(null, new StringBuilder("\"" + FileName + "\" " + CommandLine), IntPtr.Zero, IntPtr.Zero, true, Flags, EnvironmentBlock, WorkingDirectory, StartupInfo, ProcessInfo) == 0)
+							if(!bProcessInJob)
 							{
-								throw new Win32Exception();
+								throw new Win32Exception(OriginalError);
 							}
 						}
-						finally
+
+						// Allow the thread to start running
+						if(ResumeThread(ProcessInfo.hThread) == -1)
 						{
-							// Close the write ends of the handle. We don't want any other process to be able to inherit these.
-							if(StdInRead != null)
-							{
-								StdInRead.Dispose();
-								StdInRead = null;
-							}
-							if(StdOutWrite != null)
-							{
-								StdOutWrite.Dispose();
-								StdOutWrite = null;
-							}
-							if(StdErrWrite != null)
-							{
-								StdErrWrite.Dispose();
-								StdErrWrite = null;
-							}
+							throw new Win32Exception();
 						}
-					}
-
-					// Add it to our job object
-					if(AssignProcessToJobObject(Group.JobHandle, ProcessInfo.hProcess) == 0)
-					{
-						// Support for nested job objects was only addeed in Windows 8; prior to that, assigning processes to job objects would fail. Figure out if we're already in a job, and ignore the error if we are.
-						int OriginalError = Marshal.GetLastWin32Error();
-
-						bool bProcessInJob;
-						IsProcessInJob(GetCurrentProcess(), IntPtr.Zero, out bProcessInJob);
 
-						if(!bProcessInJob)
+						// If we have any input text, write it to stdin now
+						using(FileStream Stream = new FileStream(StdInWrite, FileAccess.Write, 4096, false))
 						{
-							throw new Win32Exception(OriginalError);
+							if(Input != null)
+							{
+								Stream.Write(Input, 0, Input.Length);
+								Stream.Flush();
+							}
 						}
-					}
 
-					// Allow the thread to start running
-					if(ResumeThread(ProcessInfo.hThread) == -1)
-					{
-						throw new Win32Exception();
-					}
+						// Create the stream objects for reading the process output
+						InnerStream = new FileStream(StdOutRead, FileAccess.Read, 4096, false);
+						ReadStream = new StreamReader(InnerStream, Console.OutputEncoding);
 
-					// If we have any input text, write it to stdin now
-					using(StreamWriter StdInWriter = new StreamWriter(new FileStream(StdInWrite, FileAccess.Write, 4096, false), Console.InputEncoding, 4096))
+						// Wrap the process handle in a SafeFileHandle
+						ProcessHandle = new SafeFileHandle(ProcessInfo.hProcess, true);
+					}
+					finally
 					{
-						if(!String.IsNullOrEmpty(Input))
+						if(ProcessInfo.hProcess != IntPtr.Zero && ProcessHandle == null)
 						{
-							StdInWriter.WriteLine(Input);
-							StdInWriter.Flush();
+							CloseHandle(ProcessInfo.hProcess);
+						}
+						if(ProcessInfo.hThread != IntPtr.Zero)
+						{
+							CloseHandle(ProcessInfo.hThread);
 						}
 					}
-
-					// Create the stream objects for reading the process output
-					InnerStream = new FileStream(StdOutRead, FileAccess.Read, 4096, false);
-					ReadStream = new StreamReader(InnerStream, Console.OutputEncoding);
-
-					// Wrap the process handle in a SafeFileHandle
-					ProcessHandle = new SafeFileHandle(ProcessInfo.hProcess, true);
 				}
 				finally
 				{
-					if(ProcessInfo.hProcess != IntPtr.Zero && ProcessHandle == null)
+					if(EnvironmentBlock != IntPtr.Zero)
 					{
-						CloseHandle(ProcessInfo.hProcess);
-					}
-					if(ProcessInfo.hThread != IntPtr.Zero)
-					{
-						CloseHandle(ProcessInfo.hThread);
+						Marshal.FreeHGlobal(EnvironmentBlock);
+						EnvironmentBlock = IntPtr.Zero;
 					}
 				}
 			}
-			finally
+			else
 			{
-				if(EnvironmentBlock != IntPtr.Zero)
+				// Fallback for Mono platforms
+				FrameworkProcess = new Process();
+				FrameworkProcess.StartInfo.FileName = FileName;
+				FrameworkProcess.StartInfo.Arguments = CommandLine;
+				FrameworkProcess.StartInfo.WorkingDirectory = WorkingDirectory;
+				FrameworkProcess.StartInfo.RedirectStandardInput = true;
+				FrameworkProcess.StartInfo.RedirectStandardOutput = true;
+				FrameworkProcess.StartInfo.UseShellExecute = false;
+				FrameworkProcess.StartInfo.CreateNoWindow = true;
+
+				if(Environment != null)
 				{
-					Marshal.FreeHGlobal(EnvironmentBlock);
-					EnvironmentBlock = IntPtr.Zero;
+					foreach(KeyValuePair<string, string> Pair in Environment)
+					{
+						FrameworkProcess.StartInfo.EnvironmentVariables[Pair.Key] = Pair.Value;
+					}
+				}
+
+				FrameworkProcess.Start();
+				FrameworkProcess.PriorityClass = Priority;
+
+				// If we have any input text, write it to stdin now
+				if(Input != null)
+				{
+					FrameworkProcess.StandardInput.BaseStream.Write(Input, 0, Input.Length);
+					FrameworkProcess.StandardInput.BaseStream.Close();
 				}
 			}
 		}
@@ -461,6 +517,48 @@ namespace Tools.DotNETCommon
 				StdOutRead.Dispose();
 				StdOutRead = null;
 			}
+			if(FrameworkProcess != null)
+			{
+				if(!FrameworkProcess.HasExited)
+				{
+					try
+					{
+						FrameworkProcess.Kill();
+						FrameworkProcess.WaitForExit();
+					}
+					catch
+					{
+					}
+				}
+				FrameworkProcess.Dispose();
+				FrameworkProcess = null;
+			}
+		}
+
+		/// <summary>
+		/// Reads data from the process output
+		/// </summary>
+		/// <param name="Buffer">The buffer to receive the data</param>
+		/// <param name="Offset">Offset within the buffer to write to</param>
+		/// <param name="Count">Maximum number of bytes to read</param>
+		/// <returns>Number of bytes read</returns>
+		public int Read(byte[] Buffer, int Offset, int Count)
+		{
+			// Fill the buffer, reentering managed code every 20ms to allow thread abort exceptions to be thrown
+			Task<int> ReadTask;
+			if(FrameworkProcess == null)
+			{
+				ReadTask = InnerStream.ReadAsync(Buffer, Offset, Count);
+			}
+			else
+			{
+				ReadTask = FrameworkProcess.StandardOutput.BaseStream.ReadAsync(Buffer, Offset, Count);
+			}
+			while(!ReadTask.Wait(20))
+			{
+				// Spin through managed code to allow things like ThreadAbortExceptions to be thrown.
+			}
+			return ReadTask.Result;
 		}
 
 		/// <summary>
@@ -478,19 +576,14 @@ namespace Tools.DotNETCommon
 			for(;;)
 			{
 				// Fill the buffer, reentering managed code every 20ms to allow thread abort exceptions to be thrown
-				Task<int> ReadTask = InnerStream.ReadAsync(Buffer, NumBytesInBuffer, Buffer.Length - NumBytesInBuffer);
-				while(!ReadTask.Wait(20))
-				{
-				}
-
-				// Update the buffer size
-				if(ReadTask.Result == 0)
+				int NumBytesRead = Read(Buffer, NumBytesInBuffer, Buffer.Length - NumBytesInBuffer);
+				if(NumBytesRead == 0)
 				{
 					break;
 				}
 
 				// Otherwise append to the existing buffer
-				NumBytesInBuffer += ReadTask.Result;
+				NumBytesInBuffer += NumBytesRead;
 
 				// Pull out all the complete output lines
 				int LastStartIdx = 0;
@@ -525,7 +618,15 @@ namespace Tools.DotNETCommon
 			try
 			{
 				// Busy wait for the ReadLine call to finish, so we can get interrupted by thread abort exceptions
-				Task<string> ReadLineTask = ReadStream.ReadLineAsync();
+				Task<string> ReadLineTask;
+				if(FrameworkProcess == null)
+				{
+					ReadLineTask = ReadStream.ReadLineAsync();
+				}
+				else
+				{
+					ReadLineTask = FrameworkProcess.StandardOutput.ReadLineAsync();
+				}
 				for (;;)
 				{
 					const int MillisecondsTimeout = 20;
@@ -553,36 +654,20 @@ namespace Tools.DotNETCommon
 		{
 			get
 			{
-				int Value;
-				if(GetExitCodeProcess(ProcessHandle, out Value) == 0)
+				if(FrameworkProcess == null)
 				{
-					throw new Win32Exception();
+					int Value;
+					if(GetExitCodeProcess(ProcessHandle, out Value) == 0)
+					{
+						throw new Win32Exception();
+					}
+					return Value;
+				}
+				else
+				{
+					return FrameworkProcess.ExitCode;
 				}
-				return Value;
-			}
-		}
-
-		/// <summary>
-		/// Gets the native flag for the given priority class, as passed 
-		/// </summary>
-		/// <param name="Priority">The priority value</param>
-		/// <returns>The native priority class flag</returns>
-		private static int GetNativePriorityClassFlag(ManagedProcessPriority Priority)
-		{
-			// Values below are taken from documentation on MSDN
-			switch(Priority)
-			{
-				case ManagedProcessPriority.BelowNormal:
-					// BELOW_NORMAL_PRIORITY_CLASS
-					return 0x4000;
-				case ManagedProcessPriority.Normal:
-					// NORMAL_PRIORITY_CLASS
-					return 0x0020;
-				case ManagedProcessPriority.AboveNormal:
-					// ABOVE_NORMAL_PRIORITY_CLASS
-					return 0x8000;
 			}
-			throw new Exception(String.Format("Invalid priority value: {0}", Priority));
 		}
 	}
 }
diff --git a/Engine/Source/Programs/DotNETCommon/DotNETUtilities/Perforce/PerforceChildProcess.cs b/Engine/Source/Programs/DotNETCommon/DotNETUtilities/Perforce/PerforceChildProcess.cs
--- a/Engine/Source/Programs/DotNETCommon/DotNETUtilities/Perforce/PerforceChildProcess.cs
+++ b/Engine/Source/Programs/DotNETCommon/DotNETUtilities/Perforce/PerforceChildProcess.cs
@@ -27,7 +27,8 @@ namespace Tools.DotNETCommon.Perforce
 
 	class PerforceChildProcess : IDisposable
 	{
-		Process ChildProcess;
+		ManagedProcessGroup ChildProcessGroup;
+		ManagedProcess ChildProcess;
 		byte[] Buffer;
 		int BufferPos;
 		int BufferEnd;
@@ -35,24 +36,8 @@ namespace Tools.DotNETCommon.Perforce
 
 		public PerforceChildProcess(byte[] InputData, string Format, params object[] Args)
 		{
-			ChildProcess = new Process();
-			ChildProcess.StartInfo.FileName = "p4.exe";
-			ChildProcess.StartInfo.Arguments = "-G " + String.Format(Format, Args);
-
-			ChildProcess.StartInfo.RedirectStandardError = true;
-			ChildProcess.StartInfo.RedirectStandardOutput = true;
-			ChildProcess.StartInfo.RedirectStandardInput = InputData != null;
-			ChildProcess.StartInfo.UseShellExecute = false;
-			ChildProcess.StartInfo.CreateNoWindow = true;
-
-			ChildProcess.Start();
-
-			if(InputData != null)
-			{
-				ChildProcess.StandardInput.BaseStream.Write(InputData, 0, InputData.Length);
-				ChildProcess.StandardInput.BaseStream.Close();
-			}
-
+			ChildProcessGroup = new ManagedProcessGroup();
+			ChildProcess = new ManagedProcess(ChildProcessGroup, "p4.exe", "-G " + String.Format(Format, Args), null, null, InputData, ProcessPriorityClass.Normal);
 			Buffer = new byte[2048];
 		}
 
@@ -60,21 +45,14 @@ namespace Tools.DotNETCommon.Perforce
 		{
 			if(ChildProcess != null)
 			{
-				try
-				{
-					if(!ChildProcess.HasExited)
-					{
-						ChildProcess.Kill();
-						ChildProcess.WaitForExit();
-					}
-				}
-				catch
-				{
-				}
-
 				ChildProcess.Dispose();
 				ChildProcess = null;
 			}
+			if(ChildProcessGroup != null)
+			{
+				ChildProcessGroup.Dispose();
+				ChildProcessGroup = null;
+			}
 		}
 
 		public bool IsEndOfStream()
@@ -210,7 +188,7 @@ namespace Tools.DotNETCommon.Perforce
 					BufferPos = 0;
 				}
 
-				int ReadSize = ChildProcess.StandardOutput.BaseStream.Read(Buffer, BufferEnd, Buffer.Length - BufferEnd);
+				int ReadSize = ChildProcess.Read(Buffer, BufferEnd, Buffer.Length - BufferEnd);
 				if(ReadSize == 0)
 				{
 					if(BufferPos == BufferEnd)
diff --git a/Engine/Source/Programs/DotNETCommon/DotNETUtilities/Perforce/PerforceConnection.cs b/Engine/Source/Programs/DotNETCommon/DotNETUtilities/Perforce/PerforceConnection.cs
--- a/Engine/Source/Programs/DotNETCommon/DotNETUtilities/Perforce/PerforceConnection.cs
+++ b/Engine/Source/Programs/DotNETCommon/DotNETUtilities/Perforce/PerforceConnection.cs
@@ -1486,118 +1486,6 @@ namespace Tools.DotNETCommon.Perforce
 			{
 				File.Delete(ListFileName);
 			}
-
-/*
-			using(PerforceChildProcess ChildProcess = new PerforceChildProcess("{0} print \"{1}\"", GlobalOptions, FileSpec))
-			{
-				PerforceRecordType Type;
-				if(ChildProcess.TryBeginRecord(out Type))
-				{
-					for(;;)
-					{
-						if(Type != PerforceRecordType.Stat)
-						{
-							throw new PerforceException("Unexpected record type: {0}", Type);
-						}
-
-						Dictionary<string, object> Fields = new Dictionary<string, object>();
-						ChildProcess.ReadFields(Fields);
-
-						string OutputFile = Path.Combine("D:\\WSTEST\\Content", Fields["depotFile"].ToString().Replace('/', '+'));
-//						Console.WriteLine(OutputFile);
-
-						if(File.Exists(OutputFile))
-						{
-							File.SetAttributes(OutputFile, FileAttributes.Normal);
-						}
-
-						using(FileStream OutputStream = File.Open(OutputFile, FileMode.Create, FileAccess.Write, FileShare.Read))
-						{
-							for(;;)
-							{
-								if(!ChildProcess.TryBeginRecord(out Type))
-								{
-									return;
-								}
-								if(Type != PerforceRecordType.Text && Type != PerforceRecordType.Binary)
-								{
-									break;
-								}
-
-								string Key;
-								if(!ChildProcess.TryReadKey(out Key))
-								{
-									throw new PerforceException("Foo");
-								}
-								if(Key != "data")
-								{
-									throw new PerforceException("Bar");
-								}
-								ChildProcess.ReadDataValue(OutputStream);
-
-								ChildProcess.ReadFields(Fields);
-							}
-						}
-					}
-				}
-			}*/
-/*				while(!Process.IsEndOfFile())
-				{
-					// Check that a dictionary follows
-					byte Temp = Process.ReadByte();
-					if(Temp != '{')
-					{
-						throw new PerforceException("Unexpected data while parsing marshalled output - expected '{'");
-					}
-
-					// Read all the fields in the record
-					for(;;)
-					{
-						// Read the next field type. Perforce only outputs string records. A '0' character indicates the end of the dictionary.
-						byte KeyFieldType = Process.ReadByte();
-						if(KeyFieldType == '0')
-						{
-							break;
-						}
-						else if(KeyFieldType != 's')
-						{
-							throw new PerforceException("Unexpected key field type while parsing marshalled output ({0}) - expected 's'", (int)KeyFieldType);
-						}
-
-						// Read the key
-						string Key = Process.ReadString();
-
-						// Read the value type.
-						byte ValueFieldType = Process.ReadByte();
-						if(ValueFieldType == 'i')
-						{
-							// An integer
-							string Value = Process.ReadInt32().ToString();
-							Console.WriteLine("{0}={1}", Key, Value);
-//							Record.Add(new KeyValuePair<string, object>(Key, Value));
-						}
-						else if(ValueFieldType == 's')
-						{
-							// A string (or variable length data)
-							if(Key == "data")
-							{
-								MemoryStream OutputData = new MemoryStream();
-								Process.CopyData(OutputData);
-								Console.WriteLine("{0}={1} bytes", Key, OutputData.Length);
-							}
-							else
-							{
-								string Value = Process.ReadString();
-								Console.WriteLine("{0}={1}", Key, Value);
-							}
-						}
-						else
-						{
-							throw new PerforceException("Unexpected value field type while parsing marshalled output ({0}) - expected 's'", (int)ValueFieldType);
-						}
-					}
-				}
-			}*/
 		}
 
 		#endregion
diff --git a/Engine/Source/Programs/UnrealBuildTool/System/ParallelExecutor.cs b/Engine/Source/Programs/UnrealBuildTool/System/ParallelExecutor.cs
--- a/Engine/Source/Programs/UnrealBuildTool/System/ParallelExecutor.cs
+++ b/Engine/Source/Programs/UnrealBuildTool/System/ParallelExecutor.cs
@@ -248,7 +248,7 @@ namespace UnrealBuildTool
 				Action.LogLines.Add(Action.Inner.StatusDescription);
 			}
 
-			using(ManagedProcess Process = new ManagedProcess(ProcessGroup, Action.Inner.CommandPath, Action.Inner.CommandArguments, Action.Inner.WorkingDirectory, null, null, ManagedProcessPriority.BelowNormal))
+			using(ManagedProcess Process = new ManagedProcess(ProcessGroup, Action.Inner.CommandPath, Action.Inner.CommandArguments, Action.Inner.WorkingDirectory, null, null, ProcessPriorityClass.BelowNormal))
 			{
 				Action.LogLines.AddRange(Process.ReadAllLines());
 				Action.ExitCode = Process.ExitCode;
