Fix allowing use of Job objects on Windows platforms (debug code submitted by mistake)

#rnx

diff --git a/Engine/Source/Programs/DotNETCommon/DotNETUtilities/ManagedProcess.cs b/Engine/Source/Programs/DotNETCommon/DotNETUtilities/ManagedProcess.cs
--- a/Engine/Source/Programs/DotNETCommon/DotNETUtilities/ManagedProcess.cs
+++ b/Engine/Source/Programs/DotNETCommon/DotNETUtilities/ManagedProcess.cs
@@ -81,7 +81,7 @@ namespace Tools.DotNETCommon
 		/// </summary>
 		static internal bool SupportsJobObjects
 		{
-			get { return false; }// Environment.OSVersion.Platform == PlatformID.Win32NT; }
+			get { return Environment.OSVersion.Platform == PlatformID.Win32NT || Environment.OSVersion.Platform == PlatformID.Win32S || Environment.OSVersion.Platform == PlatformID.Win32Windows; }
 		}
 
 		/// <summary>
