WorkspaceTool: Use the correct path to p4 on Mac.

#rnx

diff --git a/Engine/Source/Programs/DotNETCommon/DotNETUtilities/Perforce/PerforceChildProcess.cs b/Engine/Source/Programs/DotNETCommon/DotNETUtilities/Perforce/PerforceChildProcess.cs
--- a/Engine/Source/Programs/DotNETCommon/DotNETUtilities/Perforce/PerforceChildProcess.cs
+++ b/Engine/Source/Programs/DotNETCommon/DotNETUtilities/Perforce/PerforceChildProcess.cs
@@ -36,8 +36,19 @@ namespace Tools.DotNETCommon.Perforce
 
 		public PerforceChildProcess(byte[] InputData, string Format, params object[] Args)
 		{
+			string PerforceFileName;
+			if(Environment.OSVersion.Platform == PlatformID.Win32NT || Environment.OSVersion.Platform == PlatformID.Win32S || Environment.OSVersion.Platform == PlatformID.Win32Windows)
+			{
+				PerforceFileName = "p4.exe";
+			}
+			else
+			{
+				PerforceFileName = File.Exists("/usr/local/bin/p4")? "/usr/local/bin/p4" : "/usr/bin/p4";
+			}
+
 			ChildProcessGroup = new ManagedProcessGroup();
-			ChildProcess = new ManagedProcess(ChildProcessGroup, "p4.exe", "-G " + String.Format(Format, Args), null, null, InputData, ProcessPriorityClass.Normal);
+			ChildProcess = new ManagedProcess(ChildProcessGroup, PerforceFileName, "-G " + String.Format(Format, Args), null, null, InputData, ProcessPriorityClass.Normal);
+
 			Buffer = new byte[2048];
 		}
 
diff --git a/Engine/Source/Programs/DotNETCommon/DotNETUtilities/Perforce/PerforceConnection.cs b/Engine/Source/Programs/DotNETCommon/DotNETUtilities/Perforce/PerforceConnection.cs
--- a/Engine/Source/Programs/DotNETCommon/DotNETUtilities/Perforce/PerforceConnection.cs
+++ b/Engine/Source/Programs/DotNETCommon/DotNETUtilities/Perforce/PerforceConnection.cs
@@ -127,7 +127,7 @@ namespace Tools.DotNETCommon.Perforce
 		/// <param name="RecordHandler">Handler for each received record.</param>
 		public void Command(string CommandLine, byte[] InputData, Action<List<KeyValuePair<string, object>>> RecordHandler)
 		{
-			using(PerforceChildProcess Process = new PerforceChildProcess(InputData, "-G {0} {1}", GlobalOptions, CommandLine))
+			using(PerforceChildProcess Process = new PerforceChildProcess(InputData, "{0} {1}", GlobalOptions, CommandLine))
 			{
 				List<KeyValuePair<string, object>> Record = new List<KeyValuePair<string, object>>();
 				while(Process.TryReadRecord(Record))
