WorkspaceTool: Use readonly clients for build machines.

diff --git a/Engine/Source/Programs/DotNETCommon/DotNETUtilities/Perforce/PerforceConnection.cs b/Engine/Source/Programs/DotNETCommon/DotNETUtilities/Perforce/PerforceConnection.cs
--- a/Engine/Source/Programs/DotNETCommon/DotNETUtilities/Perforce/PerforceConnection.cs
+++ b/Engine/Source/Programs/DotNETCommon/DotNETUtilities/Perforce/PerforceConnection.cs
@@ -971,6 +971,10 @@ namespace Tools.DotNETCommon.Perforce
 			{
 				NameToValue.Add(new KeyValuePair<string, object>("LineEnd", GetEnumText(typeof(ClientLineEndings), Input.LineEnd)));
 			}
+			if(Input.Type != null)
+			{
+				NameToValue.Add(new KeyValuePair<string, object>("Type", Input.Type));
+			}
 			if(Input.Stream != null)
 			{
 				NameToValue.Add(new KeyValuePair<string, object>("Stream", Input.Stream));
