Change FileSystemReference, FileReference and DirectoryReference objects to use OrdinalIgnoreCase comparisons without creating a separate copy of the string to compare. The filesystem does not use the invariant culture, and it can produce the wrong results in some cases (the ordinal comparison is faster, too).

diff --git a/Engine/Source/Programs/AutomationTool/AutomationUtils/StagedDirectoryReference.cs b/Engine/Source/Programs/AutomationTool/AutomationUtils/StagedDirectoryReference.cs
--- a/Engine/Source/Programs/AutomationTool/AutomationUtils/StagedDirectoryReference.cs
+++ b/Engine/Source/Programs/AutomationTool/AutomationUtils/StagedDirectoryReference.cs
@@ -1,10 +1,11 @@
-﻿// Copyright 1998-2018 Epic Games, Inc. All Rights Reserved.
+// Copyright 1998-2018 Epic Games, Inc. All Rights Reserved.
 
 using System;
 using System.Collections.Generic;
 using System.Linq;
 using System.Text;
 using System.Threading.Tasks;
+using Tools.DotNETCommon;
 
 public class StagedDirectoryReference : StagedFileSystemReference, IEquatable<StagedDirectoryReference>
 {
@@ -57,7 +58,7 @@ public class StagedDirectoryReference : StagedFileSystemReference, IEquatable<St
 		}
 		else
 		{
-			return (object)B != null && A.CanonicalName == B.CanonicalName;
+			return (object)B != null && A.Name.Equals(B.Name, FileSystemReference.Comparison);
 		}
 	}
 
@@ -80,7 +81,7 @@ public class StagedDirectoryReference : StagedFileSystemReference, IEquatable<St
 	public override bool Equals(object Obj)
 	{
 		StagedDirectoryReference Other = Obj as StagedDirectoryReference;
-		return Other != null && Other.CanonicalName == CanonicalName;
+		return Other != null && Other.Name.Equals(Name, FileSystemReference.Comparison);
 	}
 
 	/// <summary>
@@ -90,7 +91,7 @@ public class StagedDirectoryReference : StagedFileSystemReference, IEquatable<St
 	/// <returns>True if the two directories are identical. Case is ignored.</returns>
 	public bool Equals(StagedDirectoryReference Other)
 	{
-		return CanonicalName == Other.CanonicalName;
+		return Name.Equals(Other.Name, FileSystemReference.Comparison);
 	}
 
 	/// <summary>
@@ -99,7 +100,7 @@ public class StagedDirectoryReference : StagedFileSystemReference, IEquatable<St
 	/// <returns>Hash code for the current object.</returns>
 	public override int GetHashCode()
 	{
-		return CanonicalName.GetHashCode();
+		return FileSystemReference.Comparer.GetHashCode(Name);
 	}
 }
 
diff --git a/Engine/Source/Programs/AutomationTool/AutomationUtils/StagedFileReference.cs b/Engine/Source/Programs/AutomationTool/AutomationUtils/StagedFileReference.cs
--- a/Engine/Source/Programs/AutomationTool/AutomationUtils/StagedFileReference.cs
+++ b/Engine/Source/Programs/AutomationTool/AutomationUtils/StagedFileReference.cs
@@ -1,4 +1,4 @@
-﻿// Copyright 1998-2018 Epic Games, Inc. All Rights Reserved.
+// Copyright 1998-2018 Epic Games, Inc. All Rights Reserved.
 
 using System;
 using System.Collections.Generic;
@@ -26,7 +26,7 @@ public class StagedFileReference : StagedFileSystemReference, IEquatable<StagedF
 	/// </summary>
 	/// <param name="Name">The name of this entity.</param>
 	/// <param name="CanonicalName">Canonical name of this entity. Should be equal to Name.ToLowerInvariant().</param>
-	private StagedFileReference(string Name, string CanonicalName) : base(Name, CanonicalName)
+	private StagedFileReference(string Name, Sanitize Sanitize) : base(Name, Sanitize)
 	{
 	}
 
@@ -52,22 +52,22 @@ public class StagedFileReference : StagedFileSystemReference, IEquatable<StagedF
 	/// <summary>
 	/// Determines if the name contains the given fragment
 	/// </summary>
-	/// <param name="Name"></param>
+	/// <param name="OtherName"></param>
 	/// <returns></returns>
-	public bool ContainsName(FileSystemName Name)
+	public bool ContainsName(FileSystemName OtherName)
 	{
 		int StartIdx = 0;
 		for(;;)
 		{
-			int Idx = CanonicalName.IndexOf(Name.CanonicalName, StartIdx);
+			int Idx = Name.IndexOf(OtherName.DisplayName, StartIdx, FileSystemReference.Comparison);
 			if(Idx == -1)
 			{
 				return false;
 			}
-			if (Idx == 0 || CanonicalName[Idx - 1] == '/')
+			if (Idx == 0 || Name[Idx - 1] == '/')
 			{
-				int EndIdx = Idx + Name.CanonicalName.Length;
-				if(EndIdx < CanonicalName.Length && CanonicalName[EndIdx] == '/')
+				int EndIdx = Idx + OtherName.DisplayName.Length;
+				if(EndIdx < Name.Length && Name[EndIdx] == '/')
 				{
 					return true;
 				}
@@ -85,9 +85,9 @@ public class StagedFileReference : StagedFileSystemReference, IEquatable<StagedF
 	/// <returns>True if the file was remapped, false otherwise</returns>
 	public static bool TryRemap(StagedFileReference InputFile, StagedDirectoryReference SourceDir, StagedDirectoryReference TargetDir, out StagedFileReference RemappedFile)
 	{
-		if (InputFile.CanonicalName.StartsWith(SourceDir.CanonicalName) && InputFile.CanonicalName.Length > SourceDir.CanonicalName.Length && InputFile.CanonicalName[SourceDir.CanonicalName.Length] == '/')
+		if (InputFile.Name.StartsWith(SourceDir.Name, FileSystemReference.Comparison) && InputFile.Name.Length > SourceDir.Name.Length && InputFile.Name[SourceDir.Name.Length] == '/')
 		{
-			RemappedFile = new StagedFileReference(TargetDir.Name + InputFile.Name.Substring(SourceDir.Name.Length), TargetDir.CanonicalName + InputFile.CanonicalName.Substring(SourceDir.CanonicalName.Length));
+			RemappedFile = new StagedFileReference(TargetDir.Name + InputFile.Name.Substring(SourceDir.Name.Length), Sanitize.None);
 			return true;
 		}
 		else
@@ -125,7 +125,7 @@ public class StagedFileReference : StagedFileSystemReference, IEquatable<StagedF
 	/// <returns>Lowercase version of this file reference.</returns>
 	public StagedFileReference ToLowerInvariant()
 	{
-		return new StagedFileReference(CanonicalName, CanonicalName);
+		return new StagedFileReference(Name.ToLowerInvariant(), Sanitize.None);
 	}
 
 	/// <summary>
@@ -142,7 +142,7 @@ public class StagedFileReference : StagedFileSystemReference, IEquatable<StagedF
 		}
 		else
 		{
-			return (object)B != null && A.CanonicalName == B.CanonicalName;
+			return (object)B != null && A.Name.Equals(B.Name, FileSystemReference.Comparison);
 		}
 	}
 
@@ -165,7 +165,7 @@ public class StagedFileReference : StagedFileSystemReference, IEquatable<StagedF
 	public override bool Equals(object Obj)
 	{
 		StagedFileReference Other = Obj as StagedFileReference;
-		return Other != null && Other.CanonicalName == CanonicalName;
+		return Other != null && Other.Name.Equals(Name, FileSystemReference.Comparison);
 	}
 
 	/// <summary>
@@ -175,7 +175,7 @@ public class StagedFileReference : StagedFileSystemReference, IEquatable<StagedF
 	/// <returns>True if the two directories are identical. Case is ignored.</returns>
 	public bool Equals(StagedFileReference Other)
 	{
-		return CanonicalName == Other.CanonicalName;
+		return Name.Equals(Other.Name, FileSystemReference.Comparison);
 	}
 
 	/// <summary>
@@ -184,7 +184,7 @@ public class StagedFileReference : StagedFileSystemReference, IEquatable<StagedF
 	/// <returns>Hash code for the current object.</returns>
 	public override int GetHashCode()
 	{
-		return CanonicalName.GetHashCode();
+		return FileSystemReference.Comparer.GetHashCode(Name);
 	}
 }
 
diff --git a/Engine/Source/Programs/AutomationTool/AutomationUtils/StagedFileSystemReference.cs b/Engine/Source/Programs/AutomationTool/AutomationUtils/StagedFileSystemReference.cs
--- a/Engine/Source/Programs/AutomationTool/AutomationUtils/StagedFileSystemReference.cs
+++ b/Engine/Source/Programs/AutomationTool/AutomationUtils/StagedFileSystemReference.cs
@@ -5,6 +5,7 @@ using System.Collections.Generic;
 using System.Linq;
 using System.Text;
 using System.Threading.Tasks;
+using Tools.DotNETCommon;
 using UnrealBuildTool;
 
 /// <summary>
@@ -13,14 +14,17 @@ using UnrealBuildTool;
 public abstract class StagedFileSystemReference
 {
 	/// <summary>
-	/// The name of this file system entity
+	/// Special value used to invoke the non-sanitizing constructor overload
 	/// </summary>
-	public readonly string Name;
+	public enum Sanitize
+	{
+		None
+	}
 
 	/// <summary>
-	/// Name of this filesystem entity in a canonical format for easy comparison
+	/// The name of this file system entity
 	/// </summary>
-	public readonly string CanonicalName;
+	public readonly string Name;
 
 	/// <summary>
 	/// Constructor
@@ -29,7 +33,6 @@ public abstract class StagedFileSystemReference
 	public StagedFileSystemReference(string InName)
 	{
 		Name = InName.Replace('\\', '/');
-		CanonicalName = Name.ToLowerInvariant();
 
 		// Make sure it's not an absolute path
 		if (Name.Length >= 2 && (Name[1] == ':' || Name[0] == '/'))
@@ -116,11 +119,27 @@ public abstract class StagedFileSystemReference
 	/// Protected constructor. Initializes to the given parameters without validation.
 	/// </summary>
 	/// <param name="Name">The name of this entity.</param>
-	/// <param name="CanonicalName">Canonical name of this entity. Should be equal to Name.ToLowerInvariant().</param>
-	protected StagedFileSystemReference(string Name, string CanonicalName)
+	/// <param name="Sanitize">Dummy argument to force overload resolution to use this implementation.</param>
+	public StagedFileSystemReference(string Name, Sanitize Sanitize)
 	{
 		this.Name = Name;
-		this.CanonicalName = CanonicalName;
+	}
+
+	/// <summary>
+	/// Checks whether this name has the given extension.
+	/// </summary>
+	/// <param name="Extension">The extension to check</param>
+	/// <returns>True if this name has the given extension, false otherwise</returns>
+	public bool HasExtension(string Extension)
+	{
+		if (Extension.Length > 0 && Extension[0] != '.')
+		{
+			return Name.Length >= Extension.Length + 1 && Name[Name.Length - Extension.Length - 1] == '.' && Name.EndsWith(Extension, FileSystemReference.Comparison);
+		}
+		else
+		{
+			return Name.EndsWith(Extension, FileSystemReference.Comparison);
+		}
 	}
 
 	/// <summary>
@@ -130,7 +149,7 @@ public abstract class StagedFileSystemReference
 	/// <returns>True if the path is under the given directory</returns>
 	public bool IsUnderDirectory(StagedDirectoryReference Directory)
 	{
-		return CanonicalName.StartsWith(Directory.CanonicalName) && (CanonicalName.Length == Directory.CanonicalName.Length || CanonicalName[Directory.CanonicalName.Length] == '/');
+		return Name.StartsWith(Directory.Name, FileSystemReference.Comparison) && (Name.Length == Directory.Name.Length || Name[Directory.Name.Length] == '/');
 	}
 
 	/// <summary>
diff --git a/Engine/Source/Programs/AutomationTool/Win/WinPlatform.Automation.cs b/Engine/Source/Programs/AutomationTool/Win/WinPlatform.Automation.cs
--- a/Engine/Source/Programs/AutomationTool/Win/WinPlatform.Automation.cs
+++ b/Engine/Source/Programs/AutomationTool/Win/WinPlatform.Automation.cs
@@ -260,7 +260,7 @@ public abstract class BaseWinPlatform : Platform
 	{
 		// Check if there are any executables being staged in this directory. Usually we only need to stage runtime dependencies next to the executable, but we may be staging
 		// other engine executables too (eg. CEF)
-		List<StagedFileReference> FilesInTargetDir = SC.FilesToStage.NonUFSFiles.Keys.Where(x => x.IsUnderDirectory(StagedBinariesDir) && (x.CanonicalName.EndsWith(".exe") || x.CanonicalName.EndsWith(".dll"))).ToList();
+		List<StagedFileReference> FilesInTargetDir = SC.FilesToStage.NonUFSFiles.Keys.Where(x => x.IsUnderDirectory(StagedBinariesDir) && (x.HasExtension(".exe") || x.HasExtension(".dll"))).ToList();
 		if(FilesInTargetDir.Count > 0)
 		{
 			Log("Copying AppLocal dependencies from {0} to {1}", BaseAppLocalDependenciesPath, StagedBinariesDir);
diff --git a/Engine/Source/Programs/DotNETCommon/DotNETUtilities/DirectoryReference.cs b/Engine/Source/Programs/DotNETCommon/DotNETUtilities/DirectoryReference.cs
--- a/Engine/Source/Programs/DotNETCommon/DotNETUtilities/DirectoryReference.cs
+++ b/Engine/Source/Programs/DotNETCommon/DotNETUtilities/DirectoryReference.cs
@@ -15,6 +15,14 @@ namespace Tools.DotNETCommon
 	[Serializable]
 	public class DirectoryReference : FileSystemReference, IEquatable<DirectoryReference>
 	{
+		/// <summary>
+		/// Special value used to invoke the non-sanitizing constructor overload
+		/// </summary>
+		public enum Sanitize
+		{
+			None
+		}
+
 		/// <summary>
 		/// Default constructor.
 		/// </summary>
@@ -36,10 +44,10 @@ namespace Tools.DotNETCommon
 		/// <summary>
 		/// Constructor for creating a directory object directly from two strings.
 		/// </summary>
-		/// <param name="InFullName"></param>
-		/// <param name="InCanonicalName"></param>
-		protected DirectoryReference(string InFullName, string InCanonicalName)
-			: base(InFullName, InCanonicalName)
+		/// <param name="InFullName">The full, sanitized path name</param>
+		/// <param name="InSanitize">Dummy argument used to resolve this overload</param>
+		public DirectoryReference(string InFullName, Sanitize InSanitize)
+			: base(InFullName)
 		{
 		}
 
@@ -91,13 +99,13 @@ namespace Tools.DotNETCommon
 					return null;
 				}
 
-				int ParentLength = CanonicalName.LastIndexOf(Path.DirectorySeparatorChar);
-				if (ParentLength == 2 && CanonicalName[1] == ':')
+				int ParentLength = FullName.LastIndexOf(Path.DirectorySeparatorChar);
+				if (ParentLength == 2 && FullName[1] == ':')
 				{
 					ParentLength++;
 				}
 
-				return new DirectoryReference(FullName.Substring(0, ParentLength), CanonicalName.Substring(0, ParentLength));
+				return new DirectoryReference(FullName.Substring(0, ParentLength), Sanitize.None);
 			}
 		}
 
@@ -108,12 +116,12 @@ namespace Tools.DotNETCommon
 		/// <returns>The full directory name containing the given file</returns>
 		public static DirectoryReference GetParentDirectory(FileReference File)
 		{
-			int ParentLength = File.CanonicalName.LastIndexOf(Path.DirectorySeparatorChar);
-			if(ParentLength == 2 && File.CanonicalName[1] == ':')
+			int ParentLength = File.FullName.LastIndexOf(Path.DirectorySeparatorChar);
+			if(ParentLength == 2 && File.FullName[1] == ':')
 			{
 				ParentLength++;
 			}
-			return new DirectoryReference(File.FullName.Substring(0, ParentLength), File.CanonicalName.Substring(0, ParentLength));
+			return new DirectoryReference(File.FullName.Substring(0, ParentLength), Sanitize.None);
 		}
 
 		/// <summary>
@@ -133,7 +141,7 @@ namespace Tools.DotNETCommon
 		/// <returns>True if this path is a root directory, false otherwise</returns>
 		public bool IsRootDirectory()
 		{
-			return CanonicalName[CanonicalName.Length - 1] == Path.DirectorySeparatorChar;
+			return FullName[FullName.Length - 1] == Path.DirectorySeparatorChar;
 		}
 
 		/// <summary>
@@ -145,7 +153,7 @@ namespace Tools.DotNETCommon
 		public static DirectoryReference Combine(DirectoryReference BaseDirectory, params string[] Fragments)
 		{
 			string FullName = FileSystemReference.CombineStrings(BaseDirectory, Fragments);
-			return new DirectoryReference(FullName, FullName.ToLowerInvariant());
+			return new DirectoryReference(FullName, Sanitize.None);
 		}
 
 		/// <summary>
@@ -162,7 +170,7 @@ namespace Tools.DotNETCommon
 			}
 			else
 			{
-				return (object)B != null && A.CanonicalName == B.CanonicalName;
+				return (object)B != null && A.FullName.Equals(B.FullName, Comparison);
 			}
 		}
 
@@ -203,7 +211,7 @@ namespace Tools.DotNETCommon
 		/// <returns></returns>
 		public override int GetHashCode()
 		{
-			return CanonicalName.GetHashCode();
+			return Comparer.GetHashCode(FullName);
 		}
 
 		/// <summary>
@@ -213,17 +221,7 @@ namespace Tools.DotNETCommon
 		/// <returns>New directory reference</returns>
 		public static DirectoryReference MakeRemote(string AbsolutePath)
 		{
-			return new DirectoryReference(AbsolutePath, AbsolutePath.ToLowerInvariant());
-		}
-
-		/// <summary>
-		/// Helper function to create a directory reference from a raw platform path. The path provided *MUST* be exactly the same as that returned by Path.GetFullPath(). 
-		/// </summary>
-		/// <param name="AbsolutePath">The absolute path in the file system</param>
-		/// <returns>New file reference</returns>
-		public static DirectoryReference MakeFromNormalizedFullPath(string AbsolutePath)
-		{
-			return new DirectoryReference(AbsolutePath, AbsolutePath.ToLowerInvariant());
+			return new DirectoryReference(AbsolutePath, Sanitize.None);
 		}
 
 		/// <summary>
@@ -294,7 +292,7 @@ namespace Tools.DotNETCommon
 		{
 			foreach (string FileName in Directory.EnumerateFiles(BaseDir.FullName))
 			{
-				yield return FileReference.MakeFromNormalizedFullPath(FileName);
+				yield return new FileReference(FileName, FileReference.Sanitize.None);
 			}
 		}
 
@@ -308,7 +306,7 @@ namespace Tools.DotNETCommon
 		{
 			foreach (string FileName in Directory.EnumerateFiles(BaseDir.FullName, Pattern))
 			{
-				yield return FileReference.MakeFromNormalizedFullPath(FileName);
+				yield return new FileReference(FileName, FileReference.Sanitize.None);
 			}
 		}
 
@@ -323,7 +321,7 @@ namespace Tools.DotNETCommon
 		{
 			foreach (string FileName in Directory.EnumerateFiles(BaseDir.FullName, Pattern, Option))
 			{
-				yield return FileReference.MakeFromNormalizedFullPath(FileName);
+				yield return new FileReference(FileName, FileReference.Sanitize.None);
 			}
 		}
 
@@ -336,7 +334,7 @@ namespace Tools.DotNETCommon
 		{
 			foreach (string DirectoryName in Directory.EnumerateDirectories(BaseDir.FullName))
 			{
-				yield return DirectoryReference.MakeFromNormalizedFullPath(DirectoryName);
+				yield return new DirectoryReference(DirectoryName, Sanitize.None);
 			}
 		}
 
@@ -350,7 +348,7 @@ namespace Tools.DotNETCommon
 		{
 			foreach (string DirectoryName in Directory.EnumerateDirectories(BaseDir.FullName, Pattern))
 			{
-				yield return DirectoryReference.MakeFromNormalizedFullPath(DirectoryName);
+				yield return new DirectoryReference(DirectoryName, Sanitize.None);
 			}
 		}
 
@@ -365,7 +363,7 @@ namespace Tools.DotNETCommon
 		{
 			foreach (string DirectoryName in Directory.EnumerateDirectories(BaseDir.FullName, Pattern, Option))
 			{
-				yield return DirectoryReference.MakeFromNormalizedFullPath(DirectoryName);
+				yield return new DirectoryReference(DirectoryName, Sanitize.None);
 			}
 		}
 
@@ -404,7 +402,7 @@ namespace Tools.DotNETCommon
 		public static DirectoryReference ReadDirectoryReference(this BinaryReader Reader)
 		{
 			string FullName = Reader.ReadString();
-			return (FullName.Length == 0) ? null : DirectoryReference.MakeFromNormalizedFullPath(FullName);
+			return (FullName.Length == 0) ? null : new DirectoryReference(FullName, DirectoryReference.Sanitize.None);
 		}
 	}
 }
diff --git a/Engine/Source/Programs/DotNETCommon/DotNETUtilities/FileReference.cs b/Engine/Source/Programs/DotNETCommon/DotNETUtilities/FileReference.cs
--- a/Engine/Source/Programs/DotNETCommon/DotNETUtilities/FileReference.cs
+++ b/Engine/Source/Programs/DotNETCommon/DotNETUtilities/FileReference.cs
@@ -15,6 +15,14 @@ namespace Tools.DotNETCommon
 	[Serializable]
 	public class FileReference : FileSystemReference, IEquatable<FileReference>
 	{
+		/// <summary>
+		/// Dummy enum to allow invoking the constructor which takes a sanitized full path
+		/// </summary>
+		public enum Sanitize
+		{
+			None
+		}
+
 		/// <summary>
 		/// Default constructor.
 		/// </summary>
@@ -40,8 +48,10 @@ namespace Tools.DotNETCommon
 		/// <summary>
 		/// Default constructor.
 		/// </summary>
-		protected FileReference(string InFullName, string InCanonicalName)
-			: base(InFullName, InCanonicalName)
+		/// <param name="InFullName">The full sanitized path</param>
+		/// <param name="InSanitize">Dummary argument to use the sanitized overload</param>
+		public FileReference(string InFullName, Sanitize InSanitize)
+			: base(InFullName)
 		{
 		}
 
@@ -99,7 +109,7 @@ namespace Tools.DotNETCommon
 		public FileReference ChangeExtension(string Extension)
 		{
 			string NewFullName = Path.ChangeExtension(FullName, Extension);
-			return new FileReference(NewFullName, NewFullName.ToLowerInvariant());
+			return new FileReference(NewFullName, Sanitize.None);
 		}
 
 		/// <summary>
@@ -120,7 +130,7 @@ namespace Tools.DotNETCommon
 		public static FileReference Combine(DirectoryReference BaseDirectory, params string[] Fragments)
 		{
 			string FullName = FileSystemReference.CombineStrings(BaseDirectory, Fragments);
-			return new FileReference(FullName, FullName.ToLowerInvariant());
+			return new FileReference(FullName, Sanitize.None);
 		}
 
 		/// <summary>
@@ -131,7 +141,7 @@ namespace Tools.DotNETCommon
 		/// <returns>The new file reference</returns>
 		public static FileReference operator +(FileReference A, string B)
 		{
-			return new FileReference(A.FullName + B, A.CanonicalName + B.ToLowerInvariant());
+			return new FileReference(A.FullName + B, Sanitize.None);
 		}
 
 		/// <summary>
@@ -148,7 +158,7 @@ namespace Tools.DotNETCommon
 			}
 			else
 			{
-				return (object)B != null && A.CanonicalName == B.CanonicalName;
+				return (object)B != null && A.FullName.Equals(B.FullName, Comparison);
 			}
 		}
 
@@ -189,7 +199,7 @@ namespace Tools.DotNETCommon
 		/// <returns></returns>
 		public override int GetHashCode()
 		{
-			return CanonicalName.GetHashCode();
+			return Comparer.GetHashCode(FullName);
 		}
 
 		/// <summary>
@@ -200,17 +210,7 @@ namespace Tools.DotNETCommon
 		/// <returns>New file reference</returns>
 		public static FileReference MakeRemote(string AbsolutePath)
 		{
-			return new FileReference(AbsolutePath, AbsolutePath.ToLowerInvariant());
-		}
-
-		/// <summary>
-		/// Helper function to create a file reference from a raw platform path. The path provided *MUST* be exactly the same as that returned by Path.GetFullPath(). 
-		/// </summary>
-		/// <param name="AbsolutePath">The absolute path in the file system</param>
-		/// <returns>New file reference</returns>
-		public static FileReference MakeFromNormalizedFullPath(string AbsolutePath)
-		{
-			return new FileReference(AbsolutePath, AbsolutePath.ToLowerInvariant());
+			return new FileReference(AbsolutePath, Sanitize.None);
 		}
 
 		/// <summary>
@@ -560,7 +560,7 @@ namespace Tools.DotNETCommon
 		public static FileReference ReadFileReference(this BinaryReader Reader)
 		{
 			string FullName = Reader.ReadString();
-			return (FullName.Length == 0) ? null : FileReference.MakeFromNormalizedFullPath(FullName);
+			return (FullName.Length == 0) ? null : new FileReference(FullName, FileReference.Sanitize.None);
 		}
 
 		/// <summary>
diff --git a/Engine/Source/Programs/DotNETCommon/DotNETUtilities/FileSystemName.cs b/Engine/Source/Programs/DotNETCommon/DotNETUtilities/FileSystemName.cs
--- a/Engine/Source/Programs/DotNETCommon/DotNETUtilities/FileSystemName.cs
+++ b/Engine/Source/Programs/DotNETCommon/DotNETUtilities/FileSystemName.cs
@@ -19,11 +19,6 @@ namespace Tools.DotNETCommon
 		/// </summary>
 		public readonly string DisplayName;
 
-		/// <summary>
-		/// The canonical form of the name
-		/// </summary>
-		public readonly string CanonicalName;
-
 		/// <summary>
 		/// Constructor
 		/// </summary>
@@ -31,7 +26,6 @@ namespace Tools.DotNETCommon
 		public FileSystemName(string DisplayName)
 		{
 			this.DisplayName = DisplayName;
-			CanonicalName = DisplayName.ToLowerInvariant();
 		}
 
 		/// <summary>
@@ -44,12 +38,10 @@ namespace Tools.DotNETCommon
 			if(Idx == Reference.FullName.Length - 1)
 			{
 				DisplayName = Reference.FullName;
-				CanonicalName = Reference.CanonicalName;
 			}
 			else
 			{
 				DisplayName = Reference.FullName.Substring(Idx + 1);
-				CanonicalName = Reference.CanonicalName.Substring(Idx + 1);
 			}
 		}
 
@@ -67,7 +59,7 @@ namespace Tools.DotNETCommon
 			}
 			else
 			{
-				return (object)B != null && A.CanonicalName == B.CanonicalName;
+				return (object)B != null && A.DisplayName.Equals(B.DisplayName, FileSystemReference.Comparison);
 			}
 		}
 
@@ -96,15 +88,15 @@ namespace Tools.DotNETCommon
         {
             if(Extension.Length == 0)
             {
-                return CanonicalName.IndexOf('.') == -1;
+                return DisplayName.IndexOf('.') == -1;
             }
             else if(Extension[0] == '.')
             {
-                return CanonicalName.EndsWith(Extension, StringComparison.InvariantCultureIgnoreCase);
+                return DisplayName.EndsWith(Extension, FileSystemReference.Comparison);
             }
             else
             {
-                return CanonicalName.Length > Extension.Length && CanonicalName[CanonicalName.Length - Extension.Length] == '.' && CanonicalName.EndsWith(Extension, StringComparison.InvariantCultureIgnoreCase);
+                return DisplayName.Length > Extension.Length && DisplayName[DisplayName.Length - Extension.Length] == '.' && DisplayName.EndsWith(Extension, FileSystemReference.Comparison);
             }
         }
 
@@ -114,7 +106,7 @@ namespace Tools.DotNETCommon
         /// <returns>Hash code for this object</returns>
         public override int GetHashCode()
 		{
-			return CanonicalName.GetHashCode();
+			return FileSystemReference.Comparer.GetHashCode(DisplayName);
 		}
 
 		/// <summary>
diff --git a/Engine/Source/Programs/DotNETCommon/DotNETUtilities/FileSystemReference.cs b/Engine/Source/Programs/DotNETCommon/DotNETUtilities/FileSystemReference.cs
--- a/Engine/Source/Programs/DotNETCommon/DotNETUtilities/FileSystemReference.cs
+++ b/Engine/Source/Programs/DotNETCommon/DotNETUtilities/FileSystemReference.cs
@@ -21,9 +21,14 @@ namespace Tools.DotNETCommon
 		public readonly string FullName;
 
 		/// <summary>
-		/// The canonical full name for this object.
+		/// The comparer to use for file system references
 		/// </summary>
-		public readonly string CanonicalName;
+		public static readonly StringComparer Comparer = StringComparer.OrdinalIgnoreCase;
+
+		/// <summary>
+		/// The comparison to use for file system references
+		/// </summary>
+		public static readonly StringComparison Comparison = StringComparison.OrdinalIgnoreCase;
 
 		/// <summary>
 		/// Direct constructor for a path
@@ -31,7 +36,6 @@ namespace Tools.DotNETCommon
 		protected FileSystemReference(string InFullName)
 		{
 			FullName = InFullName;
-			CanonicalName = InFullName.ToLowerInvariant();
 		}
 
 		/// <summary>
@@ -40,7 +44,6 @@ namespace Tools.DotNETCommon
 		protected FileSystemReference(string InFullName, string InCanonicalName)
 		{
 			FullName = InFullName;
-			CanonicalName = InCanonicalName;
 		}
 
 		/// <summary>
@@ -130,11 +133,11 @@ namespace Tools.DotNETCommon
 		{
 			if (Extension.Length > 0 && Extension[0] != '.')
 			{
-				return HasExtension("." + Extension);
+				return FullName.Length >= Extension.Length + 1 && FullName[FullName.Length - Extension.Length - 1] == '.' && FullName.EndsWith(Extension, Comparison);
 			}
 			else
 			{
-				return CanonicalName.EndsWith(Extension, StringComparison.InvariantCultureIgnoreCase);
+				return FullName.EndsWith(Extension, Comparison);
 			}
 		}
 
@@ -145,7 +148,7 @@ namespace Tools.DotNETCommon
 		/// <returns>True if this path is under the given directory</returns>
 		public bool IsUnderDirectory(DirectoryReference Other)
 		{
-			return CanonicalName.StartsWith(Other.CanonicalName) && (CanonicalName.Length == Other.CanonicalName.Length || CanonicalName[Other.CanonicalName.Length] == Path.DirectorySeparatorChar || Other.IsRootDirectory());
+			return FullName.StartsWith(Other.FullName, Comparison) && (FullName.Length == Other.FullName.Length || FullName[Other.FullName.Length] == Path.DirectorySeparatorChar || Other.IsRootDirectory());
 		}
 
 		/// <summary>
@@ -169,7 +172,7 @@ namespace Tools.DotNETCommon
 		public bool ContainsName(FileSystemName Name, int Offset, int Length)
 		{
 			// Check the substring to search is at least long enough to contain a match
-			if(Length < Name.CanonicalName.Length)
+			if(Length < Name.DisplayName.Length)
 			{
 				return false;
 			}
@@ -179,21 +182,21 @@ namespace Tools.DotNETCommon
 			for(;;)
 			{
 				// Find the next occurrence
-				MatchIdx = CanonicalName.IndexOf(Name.CanonicalName, MatchIdx, Offset + Length - MatchIdx);
+				MatchIdx = FullName.IndexOf(Name.DisplayName, MatchIdx, Offset + Length - MatchIdx, Comparison);
 				if(MatchIdx == -1)
 				{
 					return false;
 				}
 
 				// Check if the substring is a directory
-				int MatchEndIdx = MatchIdx + Name.CanonicalName.Length;
-				if(CanonicalName[MatchIdx - 1] == Path.DirectorySeparatorChar && (MatchEndIdx == CanonicalName.Length || CanonicalName[MatchEndIdx] == Path.DirectorySeparatorChar))
+				int MatchEndIdx = MatchIdx + Name.DisplayName.Length;
+				if(FullName[MatchIdx - 1] == Path.DirectorySeparatorChar && (MatchEndIdx == FullName.Length || FullName[MatchEndIdx] == Path.DirectorySeparatorChar))
 				{
 					return true;
 				}
 
 				// Move past the string that didn't match
-				MatchIdx += Name.CanonicalName.Length;
+				MatchIdx += Name.DisplayName.Length;
 			}
 		}
 
@@ -246,25 +249,25 @@ namespace Tools.DotNETCommon
 			int CommonDirectoryLength = -1;
 			for (int Idx = 0; ; Idx++)
 			{
-				if (Idx == CanonicalName.Length)
+				if (Idx == FullName.Length)
 				{
 					// The two paths are identical. Just return the "." character.
-					if (Idx == Directory.CanonicalName.Length)
+					if (Idx == Directory.FullName.Length)
 					{
 						return ".";
 					}
 
 					// Check if we're finishing on a complete directory name
-					if (Directory.CanonicalName[Idx] == Path.DirectorySeparatorChar)
+					if (Directory.FullName[Idx] == Path.DirectorySeparatorChar)
 					{
 						CommonDirectoryLength = Idx;
 					}
 					break;
 				}
-				else if (Idx == Directory.CanonicalName.Length)
+				else if (Idx == Directory.FullName.Length)
 				{
 					// Check whether the end of the directory name coincides with a boundary for the current name.
-					if (CanonicalName[Idx] == Path.DirectorySeparatorChar)
+					if (FullName[Idx] == Path.DirectorySeparatorChar)
 					{
 						CommonDirectoryLength = Idx;
 					}
@@ -273,11 +276,11 @@ namespace Tools.DotNETCommon
 				else
 				{
 					// Check the two paths match, and bail if they don't. Increase the common directory length if we've reached a separator.
-					if (CanonicalName[Idx] != Directory.CanonicalName[Idx])
+					if(String.Compare(FullName, Idx, Directory.FullName, Idx, 1, Comparison) != 0)
 					{
 						break;
 					}
-					if (CanonicalName[Idx] == Path.DirectorySeparatorChar)
+					if (FullName[Idx] == Path.DirectorySeparatorChar)
 					{
 						CommonDirectoryLength = Idx;
 					}
@@ -292,14 +295,14 @@ namespace Tools.DotNETCommon
 
 			// Append all the '..' separators to get back to the common directory, then the rest of the string to reach the target item
 			StringBuilder Result = new StringBuilder();
-			for (int Idx = CommonDirectoryLength + 1; Idx < Directory.CanonicalName.Length; Idx++)
+			for (int Idx = CommonDirectoryLength + 1; Idx < Directory.FullName.Length; Idx++)
 			{
 				// Move up a directory
 				Result.Append("..");
 				Result.Append(Path.DirectorySeparatorChar);
 
 				// Scan to the next directory separator
-				while (Idx < Directory.CanonicalName.Length && Directory.CanonicalName[Idx] != Path.DirectorySeparatorChar)
+				while (Idx < Directory.FullName.Length && Directory.FullName[Idx] != Path.DirectorySeparatorChar)
 				{
 					Idx++;
 				}
diff --git a/Engine/Source/Programs/UnrealBuildTool/Platform/Android/AndroidToolChain.cs b/Engine/Source/Programs/UnrealBuildTool/Platform/Android/AndroidToolChain.cs
--- a/Engine/Source/Programs/UnrealBuildTool/Platform/Android/AndroidToolChain.cs
+++ b/Engine/Source/Programs/UnrealBuildTool/Platform/Android/AndroidToolChain.cs
@@ -1341,7 +1341,7 @@ namespace UnrealBuildTool
 							FileArguments += GetCompileArguments_C(bDisableOptimizations);
 
 							// remove shadow variable warnings for externally included files
-							if (!SourceFile.AbsolutePath.Replace("\\", "/").StartsWith(Path.GetFullPath(UnrealBuildTool.RootDirectory.CanonicalName)))
+							if (!SourceFile.Location.IsUnderDirectory(UnrealBuildTool.RootDirectory))
 							{
 								bDisableShadowWarning = true;
 							}
diff --git a/Engine/Source/Programs/UnrealBuildTool/Platform/IOS/IOSToolChain.cs b/Engine/Source/Programs/UnrealBuildTool/Platform/IOS/IOSToolChain.cs
--- a/Engine/Source/Programs/UnrealBuildTool/Platform/IOS/IOSToolChain.cs
+++ b/Engine/Source/Programs/UnrealBuildTool/Platform/IOS/IOSToolChain.cs
@@ -1173,8 +1173,8 @@ namespace UnrealBuildTool
                 GenDebugAction.ActionHandler = new Action.BlockingActionHandler(RPCUtilHelper.RPCActionHandler);
                 FileItem DsymExporter = FileItem.GetItemByPath(Path.GetFullPath(Path.Combine(BranchDirectory, "Engine/")) + "/Binaries/Mac/" + "DsymExporter");
                 QueueFileForBatchUpload(DsymExporter);
-                QueueDirectoryForBatchUpload(DirectoryReference.MakeFromNormalizedFullPath(Path.GetFullPath(Path.Combine(BranchDirectory, "Engine/")) + "Binaries/ThirdParty/ICU/icu4c-53_1/Mac/"));
-                QueueDirectoryForBatchUpload(DirectoryReference.MakeFromNormalizedFullPath(Path.GetFullPath(Path.Combine(BranchDirectory, "Engine/")) + "Content/Internationalization/"));
+                QueueDirectoryForBatchUpload(new DirectoryReference(Path.GetFullPath(Path.Combine(BranchDirectory, "Engine/")) + "Binaries/ThirdParty/ICU/icu4c-53_1/Mac/", DirectoryReference.Sanitize.None));
+                QueueDirectoryForBatchUpload(new DirectoryReference(Path.GetFullPath(Path.Combine(BranchDirectory, "Engine/")) + "Content/Internationalization/", DirectoryReference.Sanitize.None));
             }
 
             GenDebugAction.CommandPath = "sh";
diff --git a/Engine/Source/Programs/UnrealBuildTool/Platform/Lumin/UEBuildLumin.cs b/Engine/Source/Programs/UnrealBuildTool/Platform/Lumin/UEBuildLumin.cs
--- a/Engine/Source/Programs/UnrealBuildTool/Platform/Lumin/UEBuildLumin.cs
+++ b/Engine/Source/Programs/UnrealBuildTool/Platform/Lumin/UEBuildLumin.cs
@@ -160,7 +160,7 @@ namespace UnrealBuildTool
 			CompileEnvironment.Definitions.Add("USE_ANDROID_OPENGL=0");
 			CompileEnvironment.Definitions.Add("WITH_OGGVORBIS=1");
 
-			DirectoryReference MLSDKDir = DirectoryReference.MakeFromNormalizedFullPath("$(MLSDK)");
+			DirectoryReference MLSDKDir = new DirectoryReference("$(MLSDK)", DirectoryReference.Sanitize.None);
 			CompileEnvironment.IncludePaths.SystemIncludePaths.Add(DirectoryReference.Combine(MLSDKDir, "lumin/stl/gnu-libstdc++/include"));
 			CompileEnvironment.IncludePaths.SystemIncludePaths.Add(DirectoryReference.Combine(MLSDKDir, "lumin/stl/gnu-libstdc++/include/aarch64-linux-android"));
 			CompileEnvironment.IncludePaths.SystemIncludePaths.Add(DirectoryReference.Combine(MLSDKDir, "include"));
diff --git a/Engine/Source/Programs/UnrealBuildTool/ProjectFiles/VisualStudioCode/VSCodeProjectFileGenerator.cs b/Engine/Source/Programs/UnrealBuildTool/ProjectFiles/VisualStudioCode/VSCodeProjectFileGenerator.cs
--- a/Engine/Source/Programs/UnrealBuildTool/ProjectFiles/VisualStudioCode/VSCodeProjectFileGenerator.cs
+++ b/Engine/Source/Programs/UnrealBuildTool/ProjectFiles/VisualStudioCode/VSCodeProjectFileGenerator.cs
@@ -935,7 +935,7 @@ namespace UnrealBuildTool
 				ExecutableFilename += ".app/Contents/MacOS/" + Path.GetFileName(ExecutableFilename);
 			}
 
-			return FileReference.MakeFromNormalizedFullPath(ExecutableFilename);
+			return new FileReference(ExecutableFilename);
 		}
 
 		private void WriteNativeLaunchConfig(ProjectData.Project InProject, JsonFile OutFile)
diff --git a/Engine/Source/Programs/UnrealBuildTool/System/SourceFileSearch.cs b/Engine/Source/Programs/UnrealBuildTool/System/SourceFileSearch.cs
--- a/Engine/Source/Programs/UnrealBuildTool/System/SourceFileSearch.cs
+++ b/Engine/Source/Programs/UnrealBuildTool/System/SourceFileSearch.cs
@@ -104,14 +104,14 @@ namespace UnrealBuildTool
 		static bool ShouldInclude(FileSystemReference Reference, string[] InvariantExcludedSuffixes)
 		{
 			// Ignore Mac resource fork files on non-HFS partitions
-			if (Path.GetFileName(Reference.CanonicalName).StartsWith("._"))
+			if (Path.GetFileName(Reference.FullName).StartsWith("._"))
 			{
 				return false;
 			}
 			
 			foreach (string InvariantExcludedSuffix in InvariantExcludedSuffixes)
 			{
-				if (StringUtils.FastEndsWith(Reference.CanonicalName, InvariantExcludedSuffix))
+				if (StringUtils.FastEndsWith(Reference.FullName, InvariantExcludedSuffix))
 				{
 					return false;
 				}
