UBT: Add utility functions to the log class to allow formatting errors and warnings in Visual Studio output format (eg. File(Line): warning: Message)

#rnx

diff --git a/Engine/Source/Programs/AutomationTool/BuildGraph/Script.cs b/Engine/Source/Programs/AutomationTool/BuildGraph/Script.cs
--- a/Engine/Source/Programs/AutomationTool/BuildGraph/Script.cs
+++ b/Engine/Source/Programs/AutomationTool/BuildGraph/Script.cs
@@ -81,7 +81,7 @@ namespace AutomationTool
 				{
 					if (!Document.bHasErrors)
 					{
-						Log.WriteLine(LogEventType.Error, LogFormatOptions.NoSeverityPrefix, "{0}({1}): error: {2}", File.FullName, Ex.LineNumber, Ex.Message);
+						Log.TraceError(File.FullName, Ex.LineNumber, "{0}", Ex.Message);
 						Document.bHasErrors = true;
 					}
 				}
diff --git a/Engine/Source/Programs/DotNETCommon/DotNETUtilities/Log.cs b/Engine/Source/Programs/DotNETCommon/DotNETUtilities/Log.cs
--- a/Engine/Source/Programs/DotNETCommon/DotNETUtilities/Log.cs
+++ b/Engine/Source/Programs/DotNETCommon/DotNETUtilities/Log.cs
@@ -552,6 +552,31 @@ namespace Tools.DotNETCommon
 			WriteLinePrivate(1, false, LogEventType.Error, LogFormatOptions.None, Format, Args);
 		}
 
+		/// <summary>
+		/// Writes an error message to the console, in a format suitable for Visual Studio to parse.
+		/// </summary>
+		/// <param name="File">The file containing the error</param>
+		/// <param name="Format">Message format string</param>
+		/// <param name="Args">Optional arguments</param>
+		[MethodImplAttribute(MethodImplOptions.NoInlining)]
+		public static void TraceError(FileReference File, string Format, params object[] Args)
+		{
+			WriteLinePrivate(1, false, LogEventType.Error, LogFormatOptions.NoSeverityPrefix, "{0}: error: {1}", File, String.Format(Format, Args));
+		}
+
+		/// <summary>
+		/// Writes an error message to the console, in a format suitable for Visual Studio to parse.
+		/// </summary>
+		/// <param name="File">The file containing the error</param>
+		/// <param name="Line">Line number of the error</param>
+		/// <param name="Format">Message format string</param>
+		/// <param name="Args">Optional arguments</param>
+		[MethodImplAttribute(MethodImplOptions.NoInlining)]
+		public static void TraceError(FileReference File, int Line, string Format, params object[] Args)
+		{
+			WriteLinePrivate(1, false, LogEventType.Error, LogFormatOptions.NoSeverityPrefix, "{0}({1}): error: {2}", File, Line, String.Format(Format, Args));
+		}
+
 		/// <summary>
 		/// Writes a verbose message to the console.
 		/// </summary>
@@ -586,6 +611,31 @@ namespace Tools.DotNETCommon
 			WriteLinePrivate(1, false, LogEventType.Warning, LogFormatOptions.None, Format, Args);
 		}
 
+		/// <summary>
+		/// Writes a warning message to the console.
+		/// </summary>
+		/// <param name="File">The file containing the error</param>
+		/// <param name="Format">Message format string</param>
+		/// <param name="Args">Optional arguments</param>
+		[MethodImplAttribute(MethodImplOptions.NoInlining)]
+		public static void TraceWarning(FileReference File, string Format, params object[] Args)
+		{
+			WriteLinePrivate(1, false, LogEventType.Warning, LogFormatOptions.NoSeverityPrefix, "{0}: warning: {1}", File, String.Format(Format, Args));
+		}
+
+		/// <summary>
+		/// Writes a warning message to the console.
+		/// </summary>
+		/// <param name="File">The file containing the error</param>
+		/// <param name="Line">Line number of the error</param>
+		/// <param name="Format">Message format string</param>
+		/// <param name="Args">Optional arguments</param>
+		[MethodImplAttribute(MethodImplOptions.NoInlining)]
+		public static void TraceWarning(FileReference File, int Line, string Format, params object[] Args)
+		{
+			WriteLinePrivate(1, false, LogEventType.Warning, LogFormatOptions.NoSeverityPrefix, "{0}({1}): warning: {2}", File, Line, String.Format(Format, Args));
+		}
+
 		/// <summary>
 		/// Writes a very verbose message to the console.
 		/// </summary>
diff --git a/Engine/Source/Programs/UnrealBuildTool/Configuration/UEBuildTarget.cs b/Engine/Source/Programs/UnrealBuildTool/Configuration/UEBuildTarget.cs
--- a/Engine/Source/Programs/UnrealBuildTool/Configuration/UEBuildTarget.cs
+++ b/Engine/Source/Programs/UnrealBuildTool/Configuration/UEBuildTarget.cs
@@ -3852,7 +3852,7 @@ namespace UnrealBuildTool
 				// If we can't use a shared PCH, check there's a private PCH set
 				if(RulesObject.PCHUsage == ModuleRules.PCHUsageMode.NoSharedPCHs && RulesObject.PrivatePCHHeaderFile == null)
 				{
-					Log.WriteLine(LogEventType.Warning, LogFormatOptions.NoSeverityPrefix, "{0}: warning: Modules should specify which file to use as a precompiled header (eg. PrivatePCHHeaderFile = \"Private/{1}PrivatePCH.h\") or enable shared PCHs (eg. PCHUsage = PCHUsageMode.UseSharedPCHs). '{1}' will be compiled using a shared PCH instead.", RulesObject.File, RulesObject.Name);
+					Log.TraceWarning(RulesObject.File, "Modules should specify which file to use as a precompiled header (eg. PrivatePCHHeaderFile = \"Private/{0}PrivatePCH.h\") or enable shared PCHs (eg. PCHUsage = PCHUsageMode.UseSharedPCHs). '{0}' will be compiled using a shared PCH instead.", RulesObject.Name);
 					RulesObject.PCHUsage = ModuleRules.PCHUsageMode.UseSharedPCHs;
 				}
 			}
diff --git a/Engine/Source/Programs/UnrealBuildTool/System/XmlConfigFile.cs b/Engine/Source/Programs/UnrealBuildTool/System/XmlConfigFile.cs
--- a/Engine/Source/Programs/UnrealBuildTool/System/XmlConfigFile.cs
+++ b/Engine/Source/Programs/UnrealBuildTool/System/XmlConfigFile.cs
@@ -86,7 +86,7 @@ namespace UnrealBuildTool
 				{
 					if (!ConfigFile.bHasErrors)
 					{
-						Log.WriteLine(LogEventType.Error, LogFormatOptions.NoSeverityPrefix, "{0}({1}): error: {2}", File.FullName, Ex.LineNumber, Ex.Message);
+						Log.TraceError(File, Ex.LineNumber, "{0}", Ex.Message);
 						ConfigFile.bHasErrors = true;
 					}
 				}
@@ -124,7 +124,7 @@ namespace UnrealBuildTool
 		/// <param name="Args">Standard argument for ValidationEventHandler</param>
 		void ValidationEvent(object Sender, ValidationEventArgs Args)
 		{
-			Log.WriteLine(LogEventType.Warning, LogFormatOptions.NoSeverityPrefix, "{0}({1}): warning: {2}", File.FullName, Args.Exception.LineNumber, Args.Message);
+			Log.TraceWarning(File, Args.Exception.LineNumber, "{0}", Args.Message);
 		}
 	}
 
