UAT: Consolidate list of all confidential folders in one place. FileFilter now uses the RestrictedFolders class, which has been moved out from UBT.

#rnx

diff --git a/Engine/Source/Programs/AutomationTool/AutomationUtils/DeploymentContext.cs b/Engine/Source/Programs/AutomationTool/AutomationUtils/DeploymentContext.cs
--- a/Engine/Source/Programs/AutomationTool/AutomationUtils/DeploymentContext.cs
+++ b/Engine/Source/Programs/AutomationTool/AutomationUtils/DeploymentContext.cs
@@ -421,7 +421,7 @@ public class DeploymentContext //: ProjectParams
 		{
 			RestrictedFolderNames.ExceptWith(PlatformExports.GetIncludedFolderNames(StagePlatform));
 		}
-		RestrictedFolderNames.UnionWith(FileFilter.RestrictedFolderNames);
+		RestrictedFolderNames.UnionWith(RestrictedFolders.Names);
 		RestrictedFolderNames.Remove(StageTargetPlatform.IniPlatformType.ToString());
 
 		// Read the game config files
diff --git a/Engine/Source/Programs/AutomationTool/BuildGraph/BuildGraph.cs b/Engine/Source/Programs/AutomationTool/BuildGraph/BuildGraph.cs
--- a/Engine/Source/Programs/AutomationTool/BuildGraph/BuildGraph.cs
+++ b/Engine/Source/Programs/AutomationTool/BuildGraph/BuildGraph.cs
@@ -125,8 +125,8 @@ namespace AutomationTool
 			DefaultProperties["RootDir"] = CommandUtils.RootDirectory.FullName;
 			DefaultProperties["IsBuildMachine"] = IsBuildMachine ? "true" : "false";
 			DefaultProperties["HostPlatform"] = HostPlatform.Current.HostEditorPlatform.ToString();
-			DefaultProperties["RestrictedFolderNames"] = String.Join(";", FileFilter.RestrictedFolderNames);
-			DefaultProperties["RestrictedFolderFilter"] = String.Join(";", FileFilter.RestrictedFolderNames.Select(x => String.Format(".../{0}/...", x)));
+			DefaultProperties["RestrictedFolderNames"] = String.Join(";", RestrictedFolders.Names);
+			DefaultProperties["RestrictedFolderFilter"] = String.Join(";", RestrictedFolders.Names.Select(x => String.Format(".../{0}/...", x)));
 
 			// Attempt to read existing Build Version information
 			BuildVersion Version;
diff --git a/Engine/Source/Programs/AutomationTool/Scripts/CheckRestrictedFolders.cs b/Engine/Source/Programs/AutomationTool/Scripts/CheckRestrictedFolders.cs
--- a/Engine/Source/Programs/AutomationTool/Scripts/CheckRestrictedFolders.cs
+++ b/Engine/Source/Programs/AutomationTool/Scripts/CheckRestrictedFolders.cs
@@ -26,7 +26,7 @@ namespace AutomationTool
 			}
 
 			// Find a list of restricted folders, and remove any names which are explicitly whitelisted
-			HashSet<string> RestrictedNames = new HashSet<string>(FileFilter.RestrictedFolderNames, StringComparer.InvariantCultureIgnoreCase);
+			HashSet<string> RestrictedNames = new HashSet<string>(RestrictedFolders.Names, StringComparer.InvariantCultureIgnoreCase);
 			foreach (string AllowParam in ParseParamValues("Allow"))
 			{
 				RestrictedNames.ExceptWith(AllowParam.Split('+'));
diff --git a/Engine/Source/Programs/DotNETCommon/DotNETUtilities/DotNETUtilities.csproj b/Engine/Source/Programs/DotNETCommon/DotNETUtilities/DotNETUtilities.csproj
--- a/Engine/Source/Programs/DotNETCommon/DotNETUtilities/DotNETUtilities.csproj
+++ b/Engine/Source/Programs/DotNETCommon/DotNETUtilities/DotNETUtilities.csproj
@@ -140,6 +140,7 @@
     <Compile Include="Perforce\Records\UserRecord.cs" />
     <Compile Include="Perforce\Records\WhereRecord.cs" />
     <Compile Include="Properties\AssemblyInfo.cs" />
+    <Compile Include="RestrictedFolders.cs" />
     <Compile Include="ResXResource\UEResXReader.cs" />
     <Compile Include="ResXResource\UEResXWriter.cs" />
     <Compile Include="RPCCommandHelper.cs" />
diff --git a/Engine/Source/Programs/DotNETCommon/DotNETUtilities/FileFilter.cs b/Engine/Source/Programs/DotNETCommon/DotNETUtilities/FileFilter.cs
--- a/Engine/Source/Programs/DotNETCommon/DotNETUtilities/FileFilter.cs
+++ b/Engine/Source/Programs/DotNETCommon/DotNETUtilities/FileFilter.cs
@@ -34,21 +34,6 @@ namespace Tools.DotNETCommon
 	/// </summary>
 	public class FileFilter
 	{
-		/// <summary>
-		/// List of all confidential folder names
-		/// </summary>
-		public static readonly string[] RestrictedFolderNames =
-		{
-			"EpicInternal",
-			"CarefullyRedist",
-			"NotForLicensees",
-			"NoRedist",
-			"PS4",
-			"XboxOne",
-			"Switch",
-			"Quail"
-		};
-
 		/// <summary>
 		/// Root node for the tree.
 		/// </summary>
@@ -363,7 +348,7 @@ namespace Tools.DotNETCommon
 		/// </summary>
 		public void ExcludeRestrictedFolders()
 		{
-			foreach(string RestrictedFolderName in RestrictedFolderNames)
+			foreach(string RestrictedFolderName in RestrictedFolders.Names)
 			{
 				AddRule(String.Format(".../{0}/...", RestrictedFolderName), FileFilterType.Exclude);
 			}
diff --git a/Engine/Source/Programs/DotNETCommon/DotNETUtilities/RestrictedFolders.cs b/Engine/Source/Programs/DotNETCommon/DotNETUtilities/RestrictedFolders.cs
new file mode 100644
--- /dev/null
+++ b/Engine/Source/Programs/DotNETCommon/DotNETUtilities/RestrictedFolders.cs
@@ -0,0 +1,116 @@
+// Copyright 1998-2018 Epic Games, Inc. All Rights Reserved.
+
+using System;
+using System.Collections.Generic;
+using System.IO;
+using System.Linq;
+using System.Text;
+using System.Threading.Tasks;
+using Tools.DotNETCommon;
+
+namespace Tools.DotNETCommon
+{
+	/// <summary>
+	/// Names of restricted folders. Note: The name of each entry is used to search for/create folders
+	/// </summary>
+	public enum RestrictedFolder
+	{
+		/// <summary>
+		/// Legacy. Should not be used any more.
+		/// </summary>
+		EpicInternal,
+
+		/// <summary>
+		/// Can be used by UE4 but not required
+		/// </summary>
+		CarefullyRedist,
+
+		/// <summary>
+		/// Epic Employees and Contractors
+		/// </summary>
+		NotForLicensees,
+
+		/// <summary>
+		/// Epic Employees only
+		/// </summary>
+		NoRedist,
+
+		/// <summary>
+		/// Playstation 4 source files
+		/// </summary>
+		PS4,
+
+		/// <summary>
+		/// XboxOne source files
+		/// </summary>
+		XboxOne,
+
+		/// <summary>
+		/// Switch source files
+		/// </summary>
+		Switch,
+
+		/// <summary>
+		/// Quail source files
+		/// </summary>
+		Quail,
+	}
+
+	/// <summary>
+	/// Utility functions for getting restricted folder
+	/// </summary>
+	public static class RestrictedFolders
+	{
+		/// <summary>
+		/// Public array of restricted folder names
+		/// </summary>
+		public static readonly string[] Names;
+
+		/// <summary>
+		/// Public array of restricted folder names
+		/// </summary>
+		public static readonly RestrictedFolder[] Values;
+
+		/// <summary>
+		/// Keeps a cache of folder substrings and flags
+		/// </summary>
+		static readonly Tuple<string, RestrictedFolder>[] NamesAndValues;
+		
+		/// <summary>
+		/// Static constructor
+		/// </summary>
+		static RestrictedFolders()
+		{
+			Names = Enum.GetNames(typeof(RestrictedFolder));
+			Values = (RestrictedFolder[])Enum.GetValues(typeof(RestrictedFolder));
+
+			NamesAndValues = new Tuple<string, RestrictedFolder>[Names.Length];
+			for(int Idx = 0; Idx < Names.Length; Idx++)
+			{
+				NamesAndValues[Idx] = Tuple.Create(Names[Idx], Values[Idx]);
+			}
+		}
+
+		/// <summary>
+		/// Finds all the restricted folder names relative to a base directory
+		/// </summary>
+		/// <param name="BaseDir">The base directory to check against</param>
+		/// <param name="OtherDir">The file or directory to check</param>
+		/// <returns>Array of restricted folder names</returns>
+		public static List<RestrictedFolder> FindRestrictedFolders(DirectoryReference BaseDir, DirectoryReference OtherDir)
+		{
+			List<RestrictedFolder> Folders = new List<RestrictedFolder>();
+			if(OtherDir.IsUnderDirectory(BaseDir))
+			{
+				foreach(Tuple<string, RestrictedFolder> Pair in NamesAndValues)
+				{
+					if(OtherDir.ContainsName(Pair.Item1, BaseDir.FullName.Length))
+					{
+						Folders.Add(Pair.Item2);
+					}
+				}
+			}
+			return Folders;
+		}
+	}
+}
diff --git a/Engine/Source/Programs/UnrealBuildTool/System/RestrictedFolders.cs b/Engine/Source/Programs/UnrealBuildTool/System/RestrictedFolders.cs
deleted file mode 100644
diff --git a/Engine/Source/Programs/UnrealBuildTool/UnrealBuildTool.csproj b/Engine/Source/Programs/UnrealBuildTool/UnrealBuildTool.csproj
--- a/Engine/Source/Programs/UnrealBuildTool/UnrealBuildTool.csproj
+++ b/Engine/Source/Programs/UnrealBuildTool/UnrealBuildTool.csproj
@@ -200,7 +200,6 @@
     <Compile Include="System\ObsoleteOverrideAttribute.cs" />
     <Compile Include="System\ProgressWriter.cs" />
     <Compile Include="System\RemoteExecutor.cs" />
-    <Compile Include="System\RestrictedFolders.cs" />
     <Compile Include="System\RulesAssembly.cs" />
     <Compile Include="System\ScopedTimer.cs" />
     <Compile Include="System\SourceFileWorkingSet.cs" />
