WorkspaceTool: Include a dump of raw bytes when garbage is read from the P4 process, for diagnostic purposes.

diff --git a/Engine/Source/Programs/DotNETCommon/DotNETUtilities/Perforce/PerforceChildProcess.cs b/Engine/Source/Programs/DotNETCommon/DotNETUtilities/Perforce/PerforceChildProcess.cs
--- a/Engine/Source/Programs/DotNETCommon/DotNETUtilities/Perforce/PerforceChildProcess.cs
+++ b/Engine/Source/Programs/DotNETCommon/DotNETUtilities/Perforce/PerforceChildProcess.cs
@@ -87,7 +87,17 @@ namespace Tools.DotNETCommon.Perforce
 			byte Temp = ReadByte();
 			if(Temp != '{')
 			{
-				throw new PerforceException("Unexpected data while parsing marshaled output - expected '{'");
+				StringBuilder Result = new StringBuilder("Unexpected data while parsing marshaled output - expected '{', got:");
+				for(int Idx = 0; Idx < 1024; Idx++)
+				{
+					Result.AppendFormat("{0}{1:x2}", ((Idx & 31) == 0)? "\n    " : " ", Temp);
+					if(IsEndOfStream())
+					{
+						break;
+					}
+					Temp = ReadByte();
+				}
+				throw new PerforceException(Result.ToString());
 			}
 
 			// Read all the fields in the record
