Add the Delete command to Perforce wrapper in DotNETUtilities.

#rnx

diff --git a/Engine/Source/Programs/DotNETCommon/DotNETUtilities/DotNETUtilities.csproj b/Engine/Source/Programs/DotNETCommon/DotNETUtilities/DotNETUtilities.csproj
--- a/Engine/Source/Programs/DotNETCommon/DotNETUtilities/DotNETUtilities.csproj
+++ b/Engine/Source/Programs/DotNETCommon/DotNETUtilities/DotNETUtilities.csproj
@@ -101,6 +101,8 @@
     <Compile Include="Perforce\Records\ClientSubmitOptions.cs" />
     <Compile Include="Perforce\Records\DeleteChangeOptions.cs" />
     <Compile Include="Perforce\Records\DeleteClientOptions.cs" />
+    <Compile Include="Perforce\Records\DeleteOptions.cs" />
+    <Compile Include="Perforce\Records\DeleteRecord.cs" />
     <Compile Include="Perforce\Records\DescribeFileRecord.cs" />
     <Compile Include="Perforce\Records\DescribeRecord.cs" />
     <Compile Include="Perforce\Records\EditOptions.cs" />
diff --git a/Engine/Source/Programs/DotNETCommon/DotNETUtilities/Perforce/PerforceConnection.cs b/Engine/Source/Programs/DotNETCommon/DotNETUtilities/Perforce/PerforceConnection.cs
--- a/Engine/Source/Programs/DotNETCommon/DotNETUtilities/Perforce/PerforceConnection.cs
+++ b/Engine/Source/Programs/DotNETCommon/DotNETUtilities/Perforce/PerforceConnection.cs
@@ -391,8 +391,8 @@ namespace Tools.DotNETCommon.Perforce
 						}
 
 						// Parse the subobject and add it to the list
-						object SubObject = ParseResponse(KeyValuePairs, ref Idx, RequiredChildSuffix, RecordInfo.SubElementRecordInfo);
-						List.Add(SubObject);
+						PerforceResponse Response = ParseResponse(KeyValuePairs, ref Idx, RequiredChildSuffix, RecordInfo.SubElementRecordInfo);
+						List.Add(Response.Data);
 					}
 				}
 				else
@@ -1053,6 +1053,42 @@ namespace Tools.DotNETCommon.Perforce
 
 		#endregion
 
+		#region p4 delete
+
+		/// <summary>
+		/// Execute the 'delete' command
+		/// </summary>
+		/// <param name="Options">Options for the command</param>
+		/// <param name="FileSpecs">List of file specifications to query</param>
+		/// <returns>List of responses from the server</returns>
+		public PerforceResponseList<DeleteRecord> Delete(int ChangeNumber, DeleteOptions Options, params string[] FileSpecs)
+		{
+			StringBuilder Arguments = new StringBuilder("delete");
+			if(ChangeNumber != -1)
+			{
+				Arguments.AppendFormat(" -c {0}", ChangeNumber);
+			}
+			if((Options & DeleteOptions.PreviewOnly) != 0)
+			{
+				Arguments.Append(" -n");
+			}
+			if((Options & DeleteOptions.KeepWorkspaceFiles) != 0)
+			{
+				Arguments.Append(" -k");
+			}
+			if((Options & DeleteOptions.WithoutSyncing) != 0)
+			{
+				Arguments.Append(" -v");
+			}
+			foreach(string FileSpec in FileSpecs)
+			{
+				Arguments.AppendFormat(" \"{0}\"", FileSpec);
+			}
+			return Command<DeleteRecord>(Arguments.ToString(), null);
+		}
+
+		#endregion
+
 		#region p4 describe
 
 		/// <summary>
diff --git a/Engine/Source/Programs/DotNETCommon/DotNETUtilities/Perforce/Records/DeleteOptions.cs b/Engine/Source/Programs/DotNETCommon/DotNETUtilities/Perforce/Records/DeleteOptions.cs
new file mode 100644
--- /dev/null
+++ b/Engine/Source/Programs/DotNETCommon/DotNETUtilities/Perforce/Records/DeleteOptions.cs
@@ -0,0 +1,35 @@
+using System;
+using System.Collections.Generic;
+using System.Linq;
+using System.Text;
+using System.Threading.Tasks;
+
+namespace Tools.DotNETCommon.Perforce
+{
+	/// <summary>
+	/// Options for the p4 delete command
+	/// </summary>
+	[Flags]
+	public enum DeleteOptions
+	{
+		/// <summary>
+		/// No options specified
+		/// </summary>
+		None = 0,
+
+		/// <summary>
+		/// Keep existing workspace files; mark the file as open for edit even if the file is not in the client view. 
+		/// </summary>
+		KeepWorkspaceFiles = 1,
+
+		/// <summary>
+		/// Preview which files would be opened for edit, without actually changing any files or metadata.
+		/// </summary>
+		PreviewOnly = 2,
+
+		/// <summary>
+		/// Perform the delete without syncing files
+		/// </summary>
+		WithoutSyncing = 4,
+	}
+}
diff --git a/Engine/Source/Programs/DotNETCommon/DotNETUtilities/Perforce/Records/DeleteRecord.cs b/Engine/Source/Programs/DotNETCommon/DotNETUtilities/Perforce/Records/DeleteRecord.cs
new file mode 100644
--- /dev/null
+++ b/Engine/Source/Programs/DotNETCommon/DotNETUtilities/Perforce/Records/DeleteRecord.cs
@@ -0,0 +1,53 @@
+using System;
+using System.Collections.Generic;
+using System.Linq;
+using System.Text;
+using System.Threading.Tasks;
+
+namespace Tools.DotNETCommon.Perforce
+{
+	/// <summary>
+	/// Information about a file opened for delete
+	/// </summary>
+	public class DeleteRecord
+	{
+		/// <summary>
+		/// Path to the file in the depot
+		/// </summary>
+		[PerforceTag("depotFile")]
+		public string DepotFile;
+
+		/// <summary>
+		/// Path to the file in the workspace
+		/// </summary>
+		[PerforceTag("clientFile")]
+		public string ClientFile;
+
+		/// <summary>
+		/// The working revision number of the file that was synced
+		/// </summary>
+		[PerforceTag("workRev")]
+		public int WorkingRevision;
+
+		/// <summary>
+		/// Action taken when syncing the file
+		/// </summary>
+		[PerforceTag("action")]
+		public string Action;
+
+		/// <summary>
+		/// Type of the file
+		/// </summary>
+		[PerforceTag("type")]
+		public string Type;
+
+		/// <summary>
+		/// Format this record for display in the debugger
+		/// </summary>
+		/// <returns>Summary of this revision</returns>
+		public override string ToString()
+		{
+			return DepotFile;
+		}
+	}
+}
diff --git a/Engine/Source/Programs/DotNETCommon/DotNETUtilities/Perforce/Records/ShelveRecord.cs b/Engine/Source/Programs/DotNETCommon/DotNETUtilities/Perforce/Records/ShelveRecord.cs
--- a/Engine/Source/Programs/DotNETCommon/DotNETUtilities/Perforce/Records/ShelveRecord.cs
+++ b/Engine/Source/Programs/DotNETCommon/DotNETUtilities/Perforce/Records/ShelveRecord.cs
@@ -14,7 +14,7 @@ namespace Tools.DotNETCommon.Perforce
 		/// <summary>
 		/// The changelist containing the shelved file
 		/// </summary>
-		[PerforceTag("change")]
+		[PerforceTag("change", Optional = true)]
 		public int Change;
 
 		/// <summary>
