WorkspaceTool: Include the full command line in the log for any P4 commands.

#rnx

diff --git a/Engine/Source/Programs/DotNETCommon/DotNETUtilities/Perforce/PerforceChildProcess.cs b/Engine/Source/Programs/DotNETCommon/DotNETUtilities/Perforce/PerforceChildProcess.cs
--- a/Engine/Source/Programs/DotNETCommon/DotNETUtilities/Perforce/PerforceChildProcess.cs
+++ b/Engine/Source/Programs/DotNETCommon/DotNETUtilities/Perforce/PerforceChildProcess.cs
@@ -46,8 +46,11 @@ namespace Tools.DotNETCommon.Perforce
 				PerforceFileName = File.Exists("/usr/local/bin/p4")? "/usr/local/bin/p4" : "/usr/bin/p4";
 			}
 
+			string FullArgumentList = "-G " + String.Format(Format, Args);
+			Log.TraceLog("Running {0} {1}", PerforceFileName, FullArgumentList);
+
 			ChildProcessGroup = new ManagedProcessGroup();
-			ChildProcess = new ManagedProcess(ChildProcessGroup, PerforceFileName, "-G " + String.Format(Format, Args), null, null, InputData, ProcessPriorityClass.Normal);
+			ChildProcess = new ManagedProcess(ChildProcessGroup, PerforceFileName, FullArgumentList, null, null, InputData, ProcessPriorityClass.Normal);
 
 			Buffer = new byte[2048];
 		}
