WorkspaceTool: Add a revert command, which reverts all open files in a workspaces and deletes any changelists and shelved files.

#rnx

diff --git a/Engine/Source/Programs/DotNETCommon/DotNETUtilities/DotNETUtilities.csproj b/Engine/Source/Programs/DotNETCommon/DotNETUtilities/DotNETUtilities.csproj
--- a/Engine/Source/Programs/DotNETCommon/DotNETUtilities/DotNETUtilities.csproj
+++ b/Engine/Source/Programs/DotNETCommon/DotNETUtilities/DotNETUtilities.csproj
@@ -104,6 +104,7 @@
     <Compile Include="Perforce\Records\DeleteOptions.cs" />
     <Compile Include="Perforce\Records\DeleteRecord.cs" />
     <Compile Include="Perforce\Records\DescribeFileRecord.cs" />
+    <Compile Include="Perforce\Records\DescribeOptions.cs" />
     <Compile Include="Perforce\Records\DescribeRecord.cs" />
     <Compile Include="Perforce\Records\EditOptions.cs" />
     <Compile Include="Perforce\Records\EditRecord.cs" />
@@ -117,6 +118,7 @@
     <Compile Include="Perforce\Records\InfoRecord.cs" />
     <Compile Include="Perforce\Records\IntegrationAction.cs" />
     <Compile Include="Perforce\Records\IntegrationRecord.cs" />
+    <Compile Include="Perforce\Records\OpenedOptions.cs" />
     <Compile Include="Perforce\Records\PrintRecord.cs" />
     <Compile Include="Perforce\Records\ReconcileOptions.cs" />
     <Compile Include="Perforce\Records\ReconcileRecord.cs" />
diff --git a/Engine/Source/Programs/DotNETCommon/DotNETUtilities/Perforce/PerforceConnection.cs b/Engine/Source/Programs/DotNETCommon/DotNETUtilities/Perforce/PerforceConnection.cs
--- a/Engine/Source/Programs/DotNETCommon/DotNETUtilities/Perforce/PerforceConnection.cs
+++ b/Engine/Source/Programs/DotNETCommon/DotNETUtilities/Perforce/PerforceConnection.cs
@@ -15,6 +15,11 @@ namespace Tools.DotNETCommon.Perforce
 	/// </summary>
 	public class PerforceConnection
 	{
+		/// <summary>
+		/// Constant for the default changelist, where valid.
+		/// </summary>
+		public const int DefaultChange = -2;
+
 		#region Plumbing
 
 		/// <summary>
@@ -334,6 +339,10 @@ namespace Tools.DotNETCommon.Perforce
 							{
 								FieldInfo.SetValue(NewRecord, (Value == "#none") ? 0 : int.Parse(Value.Substring(1)));
 							}
+							else if(Value == "default")
+							{
+								FieldInfo.SetValue(NewRecord, DefaultChange);
+							}
 							else
 							{
 								FieldInfo.SetValue(NewRecord, int.Parse(Value));
@@ -812,7 +821,7 @@ namespace Tools.DotNETCommon.Perforce
 			{
 				Arguments.AppendFormat(" -s {0}", GetEnumText(typeof(ChangeStatus), Status));
 			}
-			if(ClientName != null)
+			if(UserName != null)
 			{
 				Arguments.AppendFormat(" -u {0}", UserName);
 			}
@@ -1121,6 +1130,41 @@ namespace Tools.DotNETCommon.Perforce
 			return Command<DescribeRecord>(Arguments.ToString(), null);
 		}
 
+		/// <summary>
+		/// Describes a set of changelists
+		/// </summary>
+		/// <param name="ChangeNumbers">The changelist numbers to retrieve descriptions for</param>
+		/// <returns>List of responses from the server</returns>
+		public PerforceResponseList<DescribeRecord> Describe(DescribeOptions Options, int MaxResults, params int[] ChangeNumbers)
+		{
+			StringBuilder Arguments = new StringBuilder("describe -s");
+			if((Options & DescribeOptions.ShowDescriptionForRestrictedChanges) != 0)
+			{
+				Arguments.Append(" -f");
+			}
+			if((Options & DescribeOptions.Identity) != 0)
+			{
+				Arguments.Append(" -I");
+			}
+			if(MaxResults != -1)
+			{
+				Arguments.AppendFormat(" -m{0}", MaxResults);
+			}
+			if((Options & DescribeOptions.OriginalChangeNumber) != 0)
+			{
+				Arguments.Append(" -O");
+			}
+			if((Options & DescribeOptions.Shelved) != 0)
+			{
+				Arguments.Append(" -S");
+			}
+			foreach(int ChangeNumber in ChangeNumbers)
+			{
+				Arguments.AppendFormat(" {0}", ChangeNumber);
+			}
+			return Command<DescribeRecord>(Arguments.ToString(), null);
+		}
+
 		#endregion
 
 		#region p4 edit
@@ -1422,6 +1466,59 @@ namespace Tools.DotNETCommon.Perforce
 
 		#endregion
 
+		#region p4 opened
+
+		/// <summary>
+		/// Execute the 'opened' command
+		/// </summary>
+		/// <param name="Options">Options for the command</param>
+		/// <param name="ChangeNumber">List the files in pending changelist change. To list files in the default changelist, use DefaultChange.</param>
+		/// <param name="ClientName">List only files that are open in the given client</param>
+		/// <param name="UserName">List only files that are opened by the given user</param>
+		/// <param name="MaxResults">Maximum number of results to return</param>
+		/// <param name="FileSpecs">Specification for the files to list</param>
+		/// <returns>Response from the server</returns>
+		public PerforceResponseList<FStatRecord> Opened(OpenedOptions Options, int ChangeNumber, string ClientName, string UserName, int MaxResults, params string[] FileSpecs)
+		{
+			// Build the argument list
+			StringBuilder Arguments = new StringBuilder("opened");
+			if((Options & OpenedOptions.AllWorkspaces) != 0)
+			{
+				Arguments.AppendFormat(" -a");
+			}
+			if(ChangeNumber != -1)
+			{
+				Arguments.AppendFormat(" -c {0}", ChangeNumber);
+			}
+			if(ClientName != null)
+			{
+				Arguments.AppendFormat(" -C \"{0}\"", ClientName);
+			}
+			if(UserName != null)
+			{
+				Arguments.AppendFormat(" -u \"{0}\"", UserName);
+			}
+			if(MaxResults == DefaultChange)
+			{
+				Arguments.AppendFormat(" -m default");
+			}
+			else if(MaxResults != -1)
+			{
+				Arguments.AppendFormat(" -m {0}", MaxResults);
+			}
+			if((Options & OpenedOptions.ShortOutput) != 0)
+			{
+				Arguments.AppendFormat(" -s");
+			}
+			foreach(string FileSpec in FileSpecs)
+			{
+				Arguments.AppendFormat(" \"{0}\"", FileSpec);
+			}
+			return Command<FStatRecord>(Arguments.ToString(), null);
+		}
+
+		#endregion
+
 		#region p4 print
 
 		/// <summary>
diff --git a/Engine/Source/Programs/DotNETCommon/DotNETUtilities/Perforce/Records/ChangesRecord.cs b/Engine/Source/Programs/DotNETCommon/DotNETUtilities/Perforce/Records/ChangesRecord.cs
--- a/Engine/Source/Programs/DotNETCommon/DotNETUtilities/Perforce/Records/ChangesRecord.cs
+++ b/Engine/Source/Programs/DotNETCommon/DotNETUtilities/Perforce/Records/ChangesRecord.cs
@@ -50,7 +50,7 @@ namespace Tools.DotNETCommon.Perforce
 		/// <summary>
 		/// The path affected by this change.
 		/// </summary>
-		[PerforceTag("path")]
+		[PerforceTag("path", Optional = true)]
 		public string Path;
 
 		/// <summary>
diff --git a/Engine/Source/Programs/DotNETCommon/DotNETUtilities/Perforce/Records/DescribeOptions.cs b/Engine/Source/Programs/DotNETCommon/DotNETUtilities/Perforce/Records/DescribeOptions.cs
new file mode 100644
--- /dev/null
+++ b/Engine/Source/Programs/DotNETCommon/DotNETUtilities/Perforce/Records/DescribeOptions.cs
@@ -0,0 +1,40 @@
+using System;
+using System.Collections.Generic;
+using System.Linq;
+using System.Text;
+using System.Threading.Tasks;
+
+namespace Tools.DotNETCommon.Perforce
+{
+	/// <summary>
+	/// Options for the p4 describe command
+	/// </summary>
+	[Flags]
+	public enum DescribeOptions
+	{
+		/// <summary>
+		/// No options specified
+		/// </summary>
+		None = 0,
+
+		/// <summary>
+		/// Force the display of descriptions for restricted changelists. This option requires admin permission.
+		/// </summary>
+		ShowDescriptionForRestrictedChanges = 1,
+
+		/// <summary>
+		/// Specifies that the changelist number is the Identity field of a changelist.
+		/// </summary>
+		Identity = 2,
+
+		/// <summary>
+		/// If a changelist was renumbered on submit, and you know only the original changelist number, use -O and the original changelist number to describe the changelist.
+		/// </summary>
+		OriginalChangeNumber = 4,
+
+		/// <summary>
+		/// Display the names of files shelved for the specified changelist, including the diff of each file against its previous depot revision.
+		/// </summary>
+		Shelved = 8,
+	}
+}
diff --git a/Engine/Source/Programs/DotNETCommon/DotNETUtilities/Perforce/Records/DescribeRecord.cs b/Engine/Source/Programs/DotNETCommon/DotNETUtilities/Perforce/Records/DescribeRecord.cs
--- a/Engine/Source/Programs/DotNETCommon/DotNETUtilities/Perforce/Records/DescribeRecord.cs
+++ b/Engine/Source/Programs/DotNETCommon/DotNETUtilities/Perforce/Records/DescribeRecord.cs
@@ -56,7 +56,7 @@ namespace Tools.DotNETCommon.Perforce
 		/// <summary>
 		/// Narrowest path which contains all files affected by this change
 		/// </summary>
-		[PerforceTag("path")]
+		[PerforceTag("path", Optional = true)]
 		public string Path;
 
 		/// <summary>
diff --git a/Engine/Source/Programs/DotNETCommon/DotNETUtilities/Perforce/Records/OpenedOptions.cs b/Engine/Source/Programs/DotNETCommon/DotNETUtilities/Perforce/Records/OpenedOptions.cs
new file mode 100644
--- /dev/null
+++ b/Engine/Source/Programs/DotNETCommon/DotNETUtilities/Perforce/Records/OpenedOptions.cs
@@ -0,0 +1,30 @@
+using System;
+using System.Collections.Generic;
+using System.Linq;
+using System.Text;
+using System.Threading.Tasks;
+
+namespace Tools.DotNETCommon.Perforce
+{
+	/// <summary>
+	/// Options for the p4 opened command
+	/// </summary>
+	[Flags]
+	public enum OpenedOptions
+	{
+		/// <summary>
+		/// No options specified
+		/// </summary>
+		None = 0,
+
+		/// <summary>
+		/// List open files in all client workspaces
+		/// </summary>
+		AllWorkspaces = 1,
+
+		/// <summary>
+		/// Short output; do not output the revision number or file type. This option is more efficient, particularly when using the -a (all-workspaces) option at large sites.
+		/// </summary>
+		ShortOutput = 2,
+	}
+}
