UBT: Add a FileReference.WriteAllBytesIfModified() method to make it easier to update files without modifying the timestamp if they don't change.

diff --git a/Engine/Source/Programs/DotNETCommon/DotNETUtilities/ArrayUtils.cs b/Engine/Source/Programs/DotNETCommon/DotNETUtilities/ArrayUtils.cs
new file mode 100644
--- /dev/null
+++ b/Engine/Source/Programs/DotNETCommon/DotNETUtilities/ArrayUtils.cs
@@ -0,0 +1,35 @@
+using System;
+using System.Collections.Generic;
+using System.Linq;
+using System.Text;
+using System.Threading.Tasks;
+
+namespace Tools.DotNETCommon
+{
+	static class ArrayUtils
+	{
+		/// <summary>
+		/// Compares two byte arrays for equality
+		/// </summary>
+		/// <param name="A">The first byte array</param>
+		/// <param name="B">The second byte array</param>
+		/// <returns>True if the two arrays are equal, false otherwise</returns>
+		public static bool ByteArraysEqual(byte[] A, byte[] B)
+		{
+			if(A.Length != B.Length)
+			{
+				return false;
+			}
+
+			for(int Idx = 0; Idx < A.Length; Idx++)
+			{
+				if(A[Idx] != B[Idx])
+				{
+					return false;
+				}
+			}
+
+			return true;
+		}
+	}
+}
diff --git a/Engine/Source/Programs/DotNETCommon/DotNETUtilities/BinaryFormatterUtils.cs b/Engine/Source/Programs/DotNETCommon/DotNETUtilities/BinaryFormatterUtils.cs
--- a/Engine/Source/Programs/DotNETCommon/DotNETUtilities/BinaryFormatterUtils.cs
+++ b/Engine/Source/Programs/DotNETCommon/DotNETUtilities/BinaryFormatterUtils.cs
@@ -43,5 +43,22 @@ namespace Tools.DotNETCommon
 				Formatter.Serialize(Stream, Object);
 			}
 		}
+
+		/// <summary>
+		/// Saves a file to disk using the binary formatter, without updating the timestamp if it hasn't changed
+		/// </summary>
+		/// <param name="Location">File to write to</param>
+		/// <param name="Object">Object to serialize</param>
+		public static void SaveIfDifferent(FileReference Location, object Object)
+		{
+			byte[] Contents;
+			using(MemoryStream Stream = new MemoryStream())
+			{
+				BinaryFormatter Formatter = new BinaryFormatter();
+				Formatter.Serialize(Stream, Object);
+				Contents = Stream.ToArray();
+			}
+			FileReference.WriteAllBytesIfDifferent(Location, Contents);
+		}
 	}
 }
diff --git a/Engine/Source/Programs/DotNETCommon/DotNETUtilities/DotNETUtilities.csproj b/Engine/Source/Programs/DotNETCommon/DotNETUtilities/DotNETUtilities.csproj
--- a/Engine/Source/Programs/DotNETCommon/DotNETUtilities/DotNETUtilities.csproj
+++ b/Engine/Source/Programs/DotNETCommon/DotNETUtilities/DotNETUtilities.csproj
@@ -7,7 +7,7 @@
     <ProjectGuid>{5D7D66E8-8C76-4AF9-B3EC-2EF03421D730}</ProjectGuid>
     <OutputType>Library</OutputType>
     <AppDesignerFolder>Properties</AppDesignerFolder>
-    <RootNamespace>DotNETUtilities</RootNamespace>
+    <RootNamespace>Tools.DotNETCommon</RootNamespace>
     <AssemblyName>DotNETUtilities</AssemblyName>
     <TargetFrameworkVersion>v4.6.2</TargetFrameworkVersion>
     <FileAlignment>512</FileAlignment>
@@ -55,6 +55,7 @@
     <Compile Include="..\MetaData.cs">
       <Link>Properties\MetaData.cs</Link>
     </Compile>
+    <Compile Include="ArrayUtils.cs" />
     <Compile Include="AssemblyUtils.cs" />
     <Compile Include="BinaryFormatterUtils.cs" />
     <Compile Include="BinaryReaderExtensions.cs" />
diff --git a/Engine/Source/Programs/DotNETCommon/DotNETUtilities/FileReference.cs b/Engine/Source/Programs/DotNETCommon/DotNETUtilities/FileReference.cs
--- a/Engine/Source/Programs/DotNETCommon/DotNETUtilities/FileReference.cs
+++ b/Engine/Source/Programs/DotNETCommon/DotNETUtilities/FileReference.cs
@@ -463,6 +463,24 @@ namespace Tools.DotNETCommon
 			File.WriteAllBytes(Location.FullName, Contents);
 		}
 
+		/// <summary>
+		/// Writes the data to the given file, if it's different from what's there already
+		/// </summary>
+		/// <param name="Location">Location of the file</param>
+		/// <param name="Contents">Contents of the file</param>
+		public static void WriteAllBytesIfDifferent(FileReference Location, byte[] Contents)
+		{
+			if(FileReference.Exists(Location))
+			{
+				byte[] CurrentContents = FileReference.ReadAllBytes(Location);
+				if(ArrayUtils.ByteArraysEqual(Contents, CurrentContents))
+				{
+					return;
+				}
+			}
+			WriteAllBytes(Location, Contents);
+		}
+
 		/// <summary>
 		/// Writes the contents of a file
 		/// </summary>
diff --git a/Engine/Source/Programs/UnrealBuildTool/System/Utils.cs b/Engine/Source/Programs/UnrealBuildTool/System/Utils.cs
--- a/Engine/Source/Programs/UnrealBuildTool/System/Utils.cs
+++ b/Engine/Source/Programs/UnrealBuildTool/System/Utils.cs
@@ -1075,29 +1075,5 @@ namespace UnrealBuildTool
 				}
 			}
 		}
-
-		/// <summary>
-		/// Compares two byte arrays for equality
-		/// </summary>
-		/// <param name="A">The first byte array</param>
-		/// <param name="B">The second byte array</param>
-		/// <returns>True if the two arrays are equal, false otherwise</returns>
-		public static bool ByteArraysEqual(byte[] A, byte[] B)
-		{
-			if(A.Length != B.Length)
-			{
-				return false;
-			}
-
-			for(int Idx = 0; Idx < A.Length; Idx++)
-			{
-				if(A[Idx] != B[Idx])
-				{
-					return false;
-				}
-			}
-
-			return true;
-		}
 	}
 }
