Fix exception when building into a clean workspace, due to project intermediate directory not existing. This mirrors the behavior of BinaryFormatterUtils.Save() to BinaryFormatterUtils.SaveIfDifferent().

diff --git a/Engine/Source/Programs/DotNETCommon/DotNETUtilities/BinaryFormatterUtils.cs b/Engine/Source/Programs/DotNETCommon/DotNETUtilities/BinaryFormatterUtils.cs
--- a/Engine/Source/Programs/DotNETCommon/DotNETUtilities/BinaryFormatterUtils.cs
+++ b/Engine/Source/Programs/DotNETCommon/DotNETUtilities/BinaryFormatterUtils.cs
@@ -58,6 +58,8 @@ namespace Tools.DotNETCommon
 				Formatter.Serialize(Stream, Object);
 				Contents = Stream.ToArray();
 			}
+
+			DirectoryReference.CreateDirectory(Location.Directory);
 			FileReference.WriteAllBytesIfDifferent(Location, Contents);
 		}
 	}
