Fix for issue on Linux trying to create a file at a path that doesn't exist (on Windows & Mac other code results in this path being created earlier).

diff --git a/Engine/Source/Programs/DotNETCommon/DotNETUtilities/HashCollection.cs b/Engine/Source/Programs/DotNETCommon/DotNETUtilities/HashCollection.cs
--- a/Engine/Source/Programs/DotNETCommon/DotNETUtilities/HashCollection.cs
+++ b/Engine/Source/Programs/DotNETCommon/DotNETUtilities/HashCollection.cs
@@ -243,6 +243,8 @@ namespace Tools.DotNETCommon
 
 				string String = Encoding.UTF8.GetString(MS.ToArray());
 
+				Directory.CreateDirectory(Path.GetDirectoryName(OutPath));
+
 				File.WriteAllText(OutPath, String);
 			}
 		}
