Add more debug info to failures when duplicating handles.

#rnx

diff --git a/Engine/Source/Programs/DotNETCommon/DotNETUtilities/ManagedProcess.cs b/Engine/Source/Programs/DotNETCommon/DotNETUtilities/ManagedProcess.cs
--- a/Engine/Source/Programs/DotNETCommon/DotNETUtilities/ManagedProcess.cs
+++ b/Engine/Source/Programs/DotNETCommon/DotNETUtilities/ManagedProcess.cs
@@ -224,7 +224,7 @@ namespace Tools.DotNETCommon
             IntPtr hTargetProcess,
             out SafeWaitHandle targetHandle,
             int dwDesiredAccess,
-            bool bInheritHandle,
+            [MarshalAs(UnmanagedType.Bool)] bool bInheritHandle,
             int dwOptions
         );
 
@@ -366,7 +366,7 @@ namespace Tools.DotNETCommon
 								}
 								if(DuplicateHandle(GetCurrentProcess(), StdOutWrite, GetCurrentProcess(), out StdErrWrite, 0, true, DUPLICATE_SAME_ACCESS) == 0)
 								{
-									throw new Win32Exception();
+									throw new Win32Exception(String.Format("Unable to duplicate stdout handle ({0})", StdOutWrite));
 								}
 
 								// Create the new process as suspended, so we can modify it before it starts executing (and potentially preempting us)
