A couple of files that got missed in the Horde check in.

#rnx

diff --git a/Engine/Source/Programs/DotNETCommon/DotNETUtilities/FileReference.cs b/Engine/Source/Programs/DotNETCommon/DotNETUtilities/FileReference.cs
--- a/Engine/Source/Programs/DotNETCommon/DotNETUtilities/FileReference.cs
+++ b/Engine/Source/Programs/DotNETCommon/DotNETUtilities/FileReference.cs
@@ -463,6 +463,26 @@ namespace Tools.DotNETCommon
 			File.SetLastWriteTimeUtc(Location.FullName, LastWriteTimeUtc);
 		}
 
+		/// <summary>
+		/// Sets the time that the file was last accessed.
+		/// </summary>
+		/// <param name="Location">Location of the file.</param>
+		/// <param name="LastWriteTime">Last access time, in local time.</param>
+		public static void SetLastAccessTime(FileReference Location, DateTime LastWriteTime)
+		{
+			File.SetLastWriteTime(Location.FullName, LastWriteTime);
+		}
+
+		/// <summary>
+		/// Sets the time that the file was last accessed.
+		/// </summary>
+		/// <param name="Location">Location of the file.</param>
+		/// <param name="LastWriteTime">Last access time, in UTC time.</param>
+		public static void SetLastAccessTimeUtc(FileReference Location, DateTime LastWriteTimeUtc)
+		{
+			File.SetLastWriteTimeUtc(Location.FullName, LastWriteTimeUtc);
+		}
+
 		/// <summary>
 		/// Writes the contents of a file
 		/// </summary>
diff --git a/Engine/Source/Programs/DotNETCommon/DotNETUtilities/Log.cs b/Engine/Source/Programs/DotNETCommon/DotNETUtilities/Log.cs
--- a/Engine/Source/Programs/DotNETCommon/DotNETUtilities/Log.cs
+++ b/Engine/Source/Programs/DotNETCommon/DotNETUtilities/Log.cs
@@ -246,7 +246,10 @@ namespace Tools.DotNETCommon
 				Log.OutputFile = OutputFile;
 				DirectoryReference.CreateDirectory(OutputFile.Directory);
 				TextWriterTraceListener LogTraceListener = new TextWriterTraceListener(new StreamWriter(OutputFile.FullName), Name);
-				Trace.Listeners.Add(LogTraceListener);
+				lock (SyncObject)
+				{
+					Trace.Listeners.Add(LogTraceListener);
+				}
 				return LogTraceListener;
 			}
 			catch (Exception Ex)
@@ -255,6 +258,36 @@ namespace Tools.DotNETCommon
 			}
 		}
 
+		/// <summary>
+		/// Adds a <see cref="TraceListener"/> to the collection in a safe manner.
+		/// </summary>
+		/// <param name="TraceListener">The <see cref="TraceListener"/> to add.</param>
+		public static void AddTraceListener(TraceListener TraceListener)
+		{
+			lock (SyncObject)
+			{
+				if (!Trace.Listeners.Contains(TraceListener))
+				{
+					Trace.Listeners.Add(TraceListener);
+				}
+			}
+		}
+
+		/// <summary>
+		/// Removes a <see cref="TraceListener"/> from the collection in a safe manner.
+		/// </summary>
+		/// <param name="TraceListener">The <see cref="TraceListener"/> to remove.</param>
+		public static void RemoveTraceListener(TraceListener TraceListener)
+		{
+			lock (SyncObject)
+			{
+				if (Trace.Listeners.Contains(TraceListener))
+				{
+					Trace.Listeners.Remove(TraceListener);
+				}
+			}
+		}
+
 		/// <summary>
 		/// Determines if a TextWriterTraceListener has been added to the list of trace listeners
 		/// </summary>
@@ -470,6 +503,7 @@ namespace Tools.DotNETCommon
 		{
 			WriteLinePrivate(StackFramesToSkip + 1, false, Verbosity, LogFormatOptions.None, Format, Args);
 		}
+
 		/// <summary>
 		/// Mostly an internal function, but expose StackFramesToSkip to allow UAT to use existing wrapper functions and still get proper formatting.
 		/// </summary>
