UBT: Fix parsing of nullable arguments passed via command line arguments.

#rnx

diff --git a/Engine/Source/Programs/DotNETCommon/DotNETUtilities/CommandLineArguments.cs b/Engine/Source/Programs/DotNETCommon/DotNETUtilities/CommandLineArguments.cs
--- a/Engine/Source/Programs/DotNETCommon/DotNETUtilities/CommandLineArguments.cs
+++ b/Engine/Source/Programs/DotNETCommon/DotNETUtilities/CommandLineArguments.cs
@@ -593,11 +593,18 @@ namespace Tools.DotNETCommon
 					IEnumerable<CommandLineAttribute> Attributes = FieldInfo.GetCustomAttributes<CommandLineAttribute>();
 					foreach(CommandLineAttribute Attribute in Attributes)
 					{
+						// Get the inner field type, unwrapping nullable types
+						Type ValueType = FieldInfo.FieldType;
+						if(ValueType.IsGenericType && ValueType.GetGenericTypeDefinition() == typeof(Nullable<>))
+						{
+							ValueType = ValueType.GetGenericArguments()[0];
+						}
+
 						// Get the appropriate prefix for this attribute
 						string Prefix = Attribute.Prefix;
 						if(Prefix == null)
 						{
-							if(FieldInfo.FieldType == typeof(bool))
+							if(ValueType == typeof(bool))
 							{
 								Prefix = String.Format("-{0}", FieldInfo.Name);
 							}
@@ -608,7 +615,7 @@ namespace Tools.DotNETCommon
 						}
 						else
 						{
-							if(FieldInfo.FieldType != typeof(bool) && Attribute.Value == null && !Prefix.EndsWith("=") && !Prefix.EndsWith(":"))
+							if(ValueType != typeof(bool) && Attribute.Value == null && !Prefix.EndsWith("=") && !Prefix.EndsWith(":"))
 							{
 								Prefix = Prefix + "=";
 							}
