Improve error output when a file force-move fails. Show the file being moved, and actions that were taken to try and force the operation.

diff --git a/Engine/Source/Programs/DotNETCommon/DotNETUtilities/FileUtils.cs b/Engine/Source/Programs/DotNETCommon/DotNETUtilities/FileUtils.cs
--- a/Engine/Source/Programs/DotNETCommon/DotNETUtilities/FileUtils.cs
+++ b/Engine/Source/Programs/DotNETCommon/DotNETUtilities/FileUtils.cs
@@ -240,30 +240,44 @@ namespace Tools.DotNETCommon
 		/// <param name="TargetFileName">Path to the target file</param>
 		public static void ForceMoveFile(FileReference SourceLocation, FileReference TargetLocation)
 		{
+			// Try to move the file into place
 			try
 			{
 				FileReference.Move(SourceLocation, TargetLocation);
+				return;
 			}
-			catch
+			catch (Exception Ex)
 			{
+				// Try to create the target directory
+				try
+				{
+					if (!DirectoryReference.Exists(TargetLocation.Directory))
+					{
+						CreateDirectoryTree(TargetLocation.Directory);
+					}
+				}
+				catch
+				{
+				}
+
+				// Try to delete an existing file at the target location
 				try
 				{
 					if (FileReference.Exists(TargetLocation))
 					{
 						FileReference.SetAttributes(TargetLocation, FileAttributes.Normal);
 						FileReference.Delete(TargetLocation);
-					}
-					else if (DirectoryReference.Exists(TargetLocation.Directory))
-					{
-						CreateDirectoryTree(TargetLocation.Directory);
+						FileReference.Move(SourceLocation, TargetLocation);
+						return;
 					}
 				}
-				catch
+				catch (Exception DeleteEx)
 				{
+					throw new WrappedFileOrDirectoryException(new AggregateException(Ex, DeleteEx), String.Format("Unable to move {0} to {1} (also tried delete/move)", SourceLocation, TargetLocation));
 				}
 
-				// Always retry; creating the directory tree may be a race
-				FileReference.Move(SourceLocation, TargetLocation);
+				// Throw the original exception
+				throw new WrappedFileOrDirectoryException(Ex, String.Format("Unable to move {0} to {1}", SourceLocation, TargetLocation));
 			}
 		}
 
