Fix exception stack traces being stomped when starting ManagedProcess instances.

diff --git a/Engine/Source/Programs/DotNETCommon/DotNETUtilities/ManagedProcess.cs b/Engine/Source/Programs/DotNETCommon/DotNETUtilities/ManagedProcess.cs
--- a/Engine/Source/Programs/DotNETCommon/DotNETUtilities/ManagedProcess.cs
+++ b/Engine/Source/Programs/DotNETCommon/DotNETUtilities/ManagedProcess.cs
@@ -287,221 +287,256 @@ namespace Tools.DotNETCommon
 		/// <param name="WorkingDirectory">Working directory for the new process. May be null to use the current working directory.</param>
 		/// <param name="Environment">Environment variables for the new process. May be null, in which case the current process' environment is inherited</param>
 		/// <param name="Input">Text to be passed via stdin to the new process. May be null.</param>
+		/// <param name="Priority">Priority for the child process</param>
 		public ManagedProcess(ManagedProcessGroup Group, string FileName, string CommandLine, string WorkingDirectory, IReadOnlyDictionary<string, string> Environment, byte[] Input, ProcessPriorityClass Priority)
 		{
-			if(ManagedProcessGroup.SupportsJobObjects)
+			// Create the child process
+			// NOTE: Child process must be created in a separate method to avoid stomping exception callstacks (https://stackoverflow.com/a/2494150)
+			try
 			{
-				// Create the child process
-				IntPtr EnvironmentBlock = IntPtr.Zero;
-				try
+				if (ManagedProcessGroup.SupportsJobObjects)
+				{
+					CreateManagedProcessWin32(Group, FileName, CommandLine, WorkingDirectory, Environment, Input, Priority);
+				}
+				else
+				{
+					CreateManagedProcessPortable(Group, FileName, CommandLine, WorkingDirectory, Environment, Input, Priority);
+				}
+			}
+			catch (Exception Ex)
+			{
+				ExceptionUtils.AddContext(Ex, "while launching {0} {1}", FileName, CommandLine);
+				throw;
+			}
+		}
+
+		/// <summary>
+		/// Spawns a new managed process using Win32 native functions. 
+		/// </summary>
+		/// <param name="Group">The managed process group to add to</param>
+		/// <param name="FileName">Path to the executable to be run</param>
+		/// <param name="CommandLine">Command line arguments for the process</param>
+		/// <param name="WorkingDirectory">Working directory for the new process. May be null to use the current working directory.</param>
+		/// <param name="Environment">Environment variables for the new process. May be null, in which case the current process' environment is inherited</param>
+		/// <param name="Input">Text to be passed via stdin to the new process. May be null.</param>
+		/// <param name="Priority">Priority for the child process</param>
+		private void CreateManagedProcessWin32(ManagedProcessGroup Group, string FileName, string CommandLine, string WorkingDirectory, IReadOnlyDictionary<string, string> Environment, byte[] Input, ProcessPriorityClass Priority)
+		{
+			IntPtr EnvironmentBlock = IntPtr.Zero;
+			try
+			{
+				// Create the environment block for the child process, if necessary.
+				if (Environment != null)
 				{
-					// Create the environment block for the child process, if necessary.
-					if(Environment != null)
+					// The native format for the environment block is a sequence of null terminated strings with a final null terminator.
+					List<byte> EnvironmentBytes = new List<byte>();
+					foreach (KeyValuePair<string, string> Pair in Environment)
 					{
-						// The native format for the environment block is a sequence of null terminated strings with a final null terminator.
-						List<byte> EnvironmentBytes = new List<byte>();
-						foreach(KeyValuePair<string, string> Pair in Environment)
-						{
-							EnvironmentBytes.AddRange(Console.OutputEncoding.GetBytes(Pair.Key));
-							EnvironmentBytes.Add((byte)'=');
-							EnvironmentBytes.AddRange(Console.OutputEncoding.GetBytes(Pair.Value));
-							EnvironmentBytes.Add((byte)0);
-						}
+						EnvironmentBytes.AddRange(Console.OutputEncoding.GetBytes(Pair.Key));
+						EnvironmentBytes.Add((byte)'=');
+						EnvironmentBytes.AddRange(Console.OutputEncoding.GetBytes(Pair.Value));
 						EnvironmentBytes.Add((byte)0);
+					}
+					EnvironmentBytes.Add((byte)0);
+
+					// Allocate an unmanaged block of memory to store it.
+					EnvironmentBlock = Marshal.AllocHGlobal(EnvironmentBytes.Count);
+					Marshal.Copy(EnvironmentBytes.ToArray(), 0, EnvironmentBlock, EnvironmentBytes.Count);
+				}
 
-						// Allocate an unmanaged block of memory to store it.
-						EnvironmentBlock = Marshal.AllocHGlobal(EnvironmentBytes.Count);
-						Marshal.Copy(EnvironmentBytes.ToArray(), 0, EnvironmentBlock, EnvironmentBytes.Count);
+				PROCESS_INFORMATION ProcessInfo = new PROCESS_INFORMATION();
+				try
+				{
+					// Get the flags to create the new process
+					ProcessCreationFlags Flags = ProcessCreationFlags.CREATE_NO_WINDOW | ProcessCreationFlags.CREATE_SUSPENDED;
+					switch (Priority)
+					{
+						case ProcessPriorityClass.Normal:
+							Flags |= ProcessCreationFlags.NORMAL_PRIORITY_CLASS;
+							break;
+						case ProcessPriorityClass.Idle:
+							Flags |= ProcessCreationFlags.IDLE_PRIORITY_CLASS;
+							break;
+						case ProcessPriorityClass.High:
+							Flags |= ProcessCreationFlags.HIGH_PRIORITY_CLASS;
+							break;
+						case ProcessPriorityClass.RealTime:
+							Flags |= ProcessCreationFlags.REALTIME_PRIORITY_CLASS;
+							break;
+						case ProcessPriorityClass.BelowNormal:
+							Flags |= ProcessCreationFlags.BELOW_NORMAL_PRIORITY_CLASS;
+							break;
+						case ProcessPriorityClass.AboveNormal:
+							Flags |= ProcessCreationFlags.ABOVE_NORMAL_PRIORITY_CLASS;
+							break;
 					}
 
-					PROCESS_INFORMATION ProcessInfo = new PROCESS_INFORMATION();
-					try
+					// Acquire a global lock before creating inheritable handles. If multiple threads create inheritable handles at the same time, and child processes will inherit them all.
+					// Since we need to wait for output pipes to be closed (in order to consume all output), this can result in output reads not returning until all processes with the same
+					// inherited handles are closed.
+					lock (LockObject)
 					{
-						// Get the flags to create the new process
-						ProcessCreationFlags Flags = ProcessCreationFlags.CREATE_NO_WINDOW | ProcessCreationFlags.CREATE_SUSPENDED;
-						switch(Priority)
+						SafeFileHandle StdInRead = null;
+						SafeFileHandle StdOutWrite = null;
+						SafeWaitHandle StdErrWrite = null;
+						try
 						{
-							case ProcessPriorityClass.Normal:
-								Flags |= ProcessCreationFlags.NORMAL_PRIORITY_CLASS;
-								break;
-							case ProcessPriorityClass.Idle:
-								Flags |= ProcessCreationFlags.IDLE_PRIORITY_CLASS;
-								break;
-							case ProcessPriorityClass.High:
-								Flags |= ProcessCreationFlags.HIGH_PRIORITY_CLASS;
-								break;
-							case ProcessPriorityClass.RealTime:
-								Flags |= ProcessCreationFlags.REALTIME_PRIORITY_CLASS;
-								break;
-							case ProcessPriorityClass.BelowNormal:
-								Flags |= ProcessCreationFlags.BELOW_NORMAL_PRIORITY_CLASS;
-								break;
-							case ProcessPriorityClass.AboveNormal:
-								Flags |= ProcessCreationFlags.ABOVE_NORMAL_PRIORITY_CLASS;
-								break;
-						}
+							// Create stdin and stdout pipes for the child process. We'll close the handles for the child process' ends after it's been created.
+							SECURITY_ATTRIBUTES SecurityAttributes = new SECURITY_ATTRIBUTES();
+							SecurityAttributes.nLength = Marshal.SizeOf(SecurityAttributes);
+							SecurityAttributes.bInheritHandle = 1;
 
-						// Acquire a global lock before creating inheritable handles. If multiple threads create inheritable handles at the same time, and child processes will inherit them all.
-						// Since we need to wait for output pipes to be closed (in order to consume all output), this can result in output reads not returning until all processes with the same
-						// inherited handles are closed.
-						lock(LockObject)
-						{
-							SafeFileHandle StdInRead = null;
-							SafeFileHandle StdOutWrite = null;
-							SafeWaitHandle StdErrWrite = null;
-							try
+							if (CreatePipe(out StdInRead, out StdInWrite, SecurityAttributes, 0) == 0 || SetHandleInformation(StdInWrite, HANDLE_FLAG_INHERIT, 0) == 0)
 							{
-								// Create stdin and stdout pipes for the child process. We'll close the handles for the child process' ends after it's been created.
-								SECURITY_ATTRIBUTES SecurityAttributes = new SECURITY_ATTRIBUTES();
-								SecurityAttributes.nLength = Marshal.SizeOf(SecurityAttributes);
-								SecurityAttributes.bInheritHandle = 1;
-
-								if(CreatePipe(out StdInRead, out StdInWrite, SecurityAttributes, 0) == 0 || SetHandleInformation(StdInWrite, HANDLE_FLAG_INHERIT, 0) == 0)
-								{
-									throw new Win32Exception();
-								}
-								if(CreatePipe(out StdOutRead, out StdOutWrite, SecurityAttributes, 0) == 0 || SetHandleInformation(StdOutRead, HANDLE_FLAG_INHERIT, 0) == 0)
-								{
-									throw new Win32Exception();
-								}
-								if(DuplicateHandle(GetCurrentProcess(), StdOutWrite, GetCurrentProcess(), out StdErrWrite, 0, true, DUPLICATE_SAME_ACCESS) == 0)
-								{
-									throw new Win32Exception(String.Format("Unable to duplicate stdout handle ({0})", StdOutWrite));
-								}
-
-								// Create the new process as suspended, so we can modify it before it starts executing (and potentially preempting us)
-								STARTUPINFO StartupInfo = new STARTUPINFO();
-								StartupInfo.cb = Marshal.SizeOf(StartupInfo);
-								StartupInfo.hStdInput = StdInRead;
-								StartupInfo.hStdOutput = StdOutWrite;
-								StartupInfo.hStdError = StdErrWrite;
-								StartupInfo.dwFlags = STARTF_USESTDHANDLES;
-
-								if (CreateProcess(null, new StringBuilder("\"" + FileName + "\" " + CommandLine), IntPtr.Zero, IntPtr.Zero, true, Flags, EnvironmentBlock, WorkingDirectory, StartupInfo, ProcessInfo) == 0)
-								{
-									throw new Win32Exception();
-								}
+								throw new Win32Exception();
 							}
-							finally
+							if (CreatePipe(out StdOutRead, out StdOutWrite, SecurityAttributes, 0) == 0 || SetHandleInformation(StdOutRead, HANDLE_FLAG_INHERIT, 0) == 0)
 							{
-								// Close the write ends of the handle. We don't want any other process to be able to inherit these.
-								if(StdInRead != null)
-								{
-									StdInRead.Dispose();
-									StdInRead = null;
-								}
-								if(StdOutWrite != null)
-								{
-									StdOutWrite.Dispose();
-									StdOutWrite = null;
-								}
-								if(StdErrWrite != null)
-								{
-									StdErrWrite.Dispose();
-									StdErrWrite = null;
-								}
+								throw new Win32Exception();
+							}
+							if (DuplicateHandle(GetCurrentProcess(), StdOutWrite, GetCurrentProcess(), out StdErrWrite, 0, true, DUPLICATE_SAME_ACCESS) == 0)
+							{
+								throw new Win32Exception(String.Format("Unable to duplicate stdout handle ({0})", StdOutWrite.DangerousGetHandle()));
 							}
-						}
-
-						// Add it to our job object
-						if(AssignProcessToJobObject(Group.JobHandle, ProcessInfo.hProcess) == 0)
-						{
-							// Support for nested job objects was only addeed in Windows 8; prior to that, assigning processes to job objects would fail. Figure out if we're already in a job, and ignore the error if we are.
-							int OriginalError = Marshal.GetLastWin32Error();
 
-							bool bProcessInJob;
-							IsProcessInJob(GetCurrentProcess(), IntPtr.Zero, out bProcessInJob);
+							// Create the new process as suspended, so we can modify it before it starts executing (and potentially preempting us)
+							STARTUPINFO StartupInfo = new STARTUPINFO();
+							StartupInfo.cb = Marshal.SizeOf(StartupInfo);
+							StartupInfo.hStdInput = StdInRead;
+							StartupInfo.hStdOutput = StdOutWrite;
+							StartupInfo.hStdError = StdErrWrite;
+							StartupInfo.dwFlags = STARTF_USESTDHANDLES;
 
-							if(!bProcessInJob)
+							if (CreateProcess(null, new StringBuilder("\"" + FileName + "\" " + CommandLine), IntPtr.Zero, IntPtr.Zero, true, Flags, EnvironmentBlock, WorkingDirectory, StartupInfo, ProcessInfo) == 0)
 							{
-								throw new Win32Exception(OriginalError);
+								throw new Win32Exception();
 							}
 						}
-
-						// Allow the thread to start running
-						if(ResumeThread(ProcessInfo.hThread) == -1)
+						finally
 						{
-							throw new Win32Exception();
-						}
-
-						// If we have any input text, write it to stdin now
-						using(FileStream Stream = new FileStream(StdInWrite, FileAccess.Write, 4096, false))
-						{
-							if(Input != null)
+							// Close the write ends of the handle. We don't want any other process to be able to inherit these.
+							if (StdInRead != null)
+							{
+								StdInRead.Dispose();
+								StdInRead = null;
+							}
+							if (StdOutWrite != null)
+							{
+								StdOutWrite.Dispose();
+								StdOutWrite = null;
+							}
+							if (StdErrWrite != null)
 							{
-								Stream.Write(Input, 0, Input.Length);
-								Stream.Flush();
+								StdErrWrite.Dispose();
+								StdErrWrite = null;
 							}
 						}
-
-						// Create the stream objects for reading the process output
-						InnerStream = new FileStream(StdOutRead, FileAccess.Read, 4096, false);
-						ReadStream = new StreamReader(InnerStream, Console.OutputEncoding);
-
-						// Wrap the process handle in a SafeFileHandle
-						ProcessHandle = new SafeFileHandle(ProcessInfo.hProcess, true);
 					}
-					finally
+
+					// Add it to our job object
+					if (AssignProcessToJobObject(Group.JobHandle, ProcessInfo.hProcess) == 0)
 					{
-						if(ProcessInfo.hProcess != IntPtr.Zero && ProcessHandle == null)
+						// Support for nested job objects was only addeed in Windows 8; prior to that, assigning processes to job objects would fail. Figure out if we're already in a job, and ignore the error if we are.
+						int OriginalError = Marshal.GetLastWin32Error();
+
+						bool bProcessInJob;
+						IsProcessInJob(GetCurrentProcess(), IntPtr.Zero, out bProcessInJob);
+
+						if (!bProcessInJob)
 						{
-							CloseHandle(ProcessInfo.hProcess);
+							throw new Win32Exception(OriginalError);
 						}
-						if(ProcessInfo.hThread != IntPtr.Zero)
+					}
+
+					// Allow the thread to start running
+					if (ResumeThread(ProcessInfo.hThread) == -1)
+					{
+						throw new Win32Exception();
+					}
+
+					// If we have any input text, write it to stdin now
+					using (FileStream Stream = new FileStream(StdInWrite, FileAccess.Write, 4096, false))
+					{
+						if (Input != null)
 						{
-							CloseHandle(ProcessInfo.hThread);
+							Stream.Write(Input, 0, Input.Length);
+							Stream.Flush();
 						}
 					}
-				}
-				catch (Exception Ex)
-				{
-					ExceptionUtils.AddContext(Ex, "while launching {0} {1}", FileName, CommandLine);
-					throw;
+
+					// Create the stream objects for reading the process output
+					InnerStream = new FileStream(StdOutRead, FileAccess.Read, 4096, false);
+					ReadStream = new StreamReader(InnerStream, Console.OutputEncoding);
+
+					// Wrap the process handle in a SafeFileHandle
+					ProcessHandle = new SafeFileHandle(ProcessInfo.hProcess, true);
 				}
 				finally
 				{
-					if(EnvironmentBlock != IntPtr.Zero)
+					if (ProcessInfo.hProcess != IntPtr.Zero && ProcessHandle == null)
+					{
+						CloseHandle(ProcessInfo.hProcess);
+					}
+					if (ProcessInfo.hThread != IntPtr.Zero)
 					{
-						Marshal.FreeHGlobal(EnvironmentBlock);
-						EnvironmentBlock = IntPtr.Zero;
+						CloseHandle(ProcessInfo.hThread);
 					}
 				}
 			}
-			else
+			finally
 			{
-				// Fallback for Mono platforms
-				FrameworkProcess = new Process();
-				FrameworkProcess.StartInfo.FileName = FileName;
-				FrameworkProcess.StartInfo.Arguments = CommandLine;
-				FrameworkProcess.StartInfo.WorkingDirectory = WorkingDirectory;
-				FrameworkProcess.StartInfo.RedirectStandardInput = true;
-				FrameworkProcess.StartInfo.RedirectStandardOutput = true;
-				FrameworkProcess.StartInfo.UseShellExecute = false;
-				FrameworkProcess.StartInfo.CreateNoWindow = true;
-
-				if(Environment != null)
+				if (EnvironmentBlock != IntPtr.Zero)
 				{
-					foreach(KeyValuePair<string, string> Pair in Environment)
-					{
-						FrameworkProcess.StartInfo.EnvironmentVariables[Pair.Key] = Pair.Value;
-					}
+					Marshal.FreeHGlobal(EnvironmentBlock);
+					EnvironmentBlock = IntPtr.Zero;
 				}
+			}
+		}
 
-				FrameworkProcess.Start();
-				try
+		/// <summary>
+		/// Spawns a new managed process using Win32 native functions. 
+		/// </summary>
+		/// <param name="Group">The managed process group to add to</param>
+		/// <param name="FileName">Path to the executable to be run</param>
+		/// <param name="CommandLine">Command line arguments for the process</param>
+		/// <param name="WorkingDirectory">Working directory for the new process. May be null to use the current working directory.</param>
+		/// <param name="Environment">Environment variables for the new process. May be null, in which case the current process' environment is inherited</param>
+		/// <param name="Input">Text to be passed via stdin to the new process. May be null.</param>
+		/// <param name="Priority">Priority for the child process</param>
+		private void CreateManagedProcessPortable(ManagedProcessGroup Group, string FileName, string CommandLine, string WorkingDirectory, IReadOnlyDictionary<string, string> Environment, byte[] Input, ProcessPriorityClass Priority)
+		{
+			// Fallback for Mono platforms
+			FrameworkProcess = new Process();
+			FrameworkProcess.StartInfo.FileName = FileName;
+			FrameworkProcess.StartInfo.Arguments = CommandLine;
+			FrameworkProcess.StartInfo.WorkingDirectory = WorkingDirectory;
+			FrameworkProcess.StartInfo.RedirectStandardInput = true;
+			FrameworkProcess.StartInfo.RedirectStandardOutput = true;
+			FrameworkProcess.StartInfo.UseShellExecute = false;
+			FrameworkProcess.StartInfo.CreateNoWindow = true;
+
+			if (Environment != null)
+			{
+				foreach (KeyValuePair<string, string> Pair in Environment)
 				{
-					FrameworkProcess.PriorityClass = Priority;
+					FrameworkProcess.StartInfo.EnvironmentVariables[Pair.Key] = Pair.Value;
 				}
-				catch
-				{
+			}
 
-				}
-				// If we have any input text, write it to stdin now
-				if (Input != null)
-				{
-					FrameworkProcess.StandardInput.BaseStream.Write(Input, 0, Input.Length);
-					FrameworkProcess.StandardInput.BaseStream.Close();
-				}
+			FrameworkProcess.Start();
+			try
+			{
+				FrameworkProcess.PriorityClass = Priority;
+			}
+			catch
+			{
+
+			}
+			// If we have any input text, write it to stdin now
+			if (Input != null)
+			{
+				FrameworkProcess.StandardInput.BaseStream.Write(Input, 0, Input.Length);
+				FrameworkProcess.StandardInput.BaseStream.Close();
 			}
 		}
 
