UBT: Add unique error messages for every site that can throw an exception when creating managed processes. Also retry creating a process if it fails with ERROR_ACCESS_DENIED, which we seem to be seeing intermittently on the builders.

diff --git a/Engine/Source/Programs/DotNETCommon/DotNETUtilities/DotNETUtilities.csproj b/Engine/Source/Programs/DotNETCommon/DotNETUtilities/DotNETUtilities.csproj
--- a/Engine/Source/Programs/DotNETCommon/DotNETUtilities/DotNETUtilities.csproj
+++ b/Engine/Source/Programs/DotNETCommon/DotNETUtilities/DotNETUtilities.csproj
@@ -165,6 +165,7 @@
     <Compile Include="StringUtils.cs" />
     <Compile Include="ThreadPoolWorkQueue.cs" />
     <Compile Include="HashCollection.cs" />
+    <Compile Include="Win32ExceptionWithCode.cs" />
   </ItemGroup>
   <Import Project="$(MSBuildToolsPath)\Microsoft.CSharp.targets" />
   <!-- To modify your build process, add your task inside one of the targets below and uncomment it. 
diff --git a/Engine/Source/Programs/DotNETCommon/DotNETUtilities/ManagedProcess.cs b/Engine/Source/Programs/DotNETCommon/DotNETUtilities/ManagedProcess.cs
--- a/Engine/Source/Programs/DotNETCommon/DotNETUtilities/ManagedProcess.cs
+++ b/Engine/Source/Programs/DotNETCommon/DotNETUtilities/ManagedProcess.cs
@@ -242,6 +242,8 @@ namespace Tools.DotNETCommon
 		[DllImport("kernel32.dll", SetLastError=true)]
 		static extern UInt32 WaitForSingleObject(SafeHandleZeroOrMinusOneIsInvalid hHandle, UInt32 dwMilliseconds);
 
+		const int ERROR_ACCESS_DENIED = 5;
+
 		/// <summary>
 		/// Handle for the child process.
 		/// </summary>
@@ -388,15 +390,15 @@ namespace Tools.DotNETCommon
 
 							if (CreatePipe(out StdInRead, out StdInWrite, SecurityAttributes, 0) == 0 || SetHandleInformation(StdInWrite, HANDLE_FLAG_INHERIT, 0) == 0)
 							{
-								throw new Win32Exception();
+								throw new Win32ExceptionWithCode("Unable to create stdin pipe");
 							}
 							if (CreatePipe(out StdOutRead, out StdOutWrite, SecurityAttributes, 0) == 0 || SetHandleInformation(StdOutRead, HANDLE_FLAG_INHERIT, 0) == 0)
 							{
-								throw new Win32Exception();
+								throw new Win32ExceptionWithCode("Unable to create stdout pipe");
 							}
 							if (DuplicateHandle(GetCurrentProcess(), StdOutWrite, GetCurrentProcess(), out StdErrWrite, 0, true, DUPLICATE_SAME_ACCESS) == 0)
 							{
-								throw new Win32Exception(String.Format("Unable to duplicate stdout handle ({0})", StdOutWrite.DangerousGetHandle()));
+								throw new Win32ExceptionWithCode("Unable to duplicate stdout handle");
 							}
 
 							// Create the new process as suspended, so we can modify it before it starts executing (and potentially preempting us)
@@ -407,9 +409,22 @@ namespace Tools.DotNETCommon
 							StartupInfo.hStdError = StdErrWrite;
 							StartupInfo.dwFlags = STARTF_USESTDHANDLES;
 
-							if (CreateProcess(null, new StringBuilder("\"" + FileName + "\" " + CommandLine), IntPtr.Zero, IntPtr.Zero, true, Flags, EnvironmentBlock, WorkingDirectory, StartupInfo, ProcessInfo) == 0)
+							// Under heavy load (ie. spawning large number of processes, typically Clang) we see CreateProcess very occasionally failing with ERROR_ACCESS_DENIED.
+							int[] RetryDelay = { 100, 200, 1000, 5000 };
+							for(int AttemptIdx = 0; ; AttemptIdx++)
 							{
-								throw new Win32Exception();
+								if (CreateProcess(null, new StringBuilder("\"" + FileName + "\" " + CommandLine), IntPtr.Zero, IntPtr.Zero, true, Flags, EnvironmentBlock, WorkingDirectory, StartupInfo, ProcessInfo) != 0)
+								{
+									break;
+								}
+
+								if (Marshal.GetLastWin32Error() != ERROR_ACCESS_DENIED || AttemptIdx >= RetryDelay.Length)
+								{
+									throw new Win32ExceptionWithCode("Unable to create process");
+								}
+
+								Log.TraceWarning("Unable to create process {0} (access denied); waiting {1}ms to retry...", FileName, RetryDelay[AttemptIdx]);
+								Thread.Sleep(RetryDelay[AttemptIdx]);
 							}
 						}
 						finally
@@ -444,14 +459,14 @@ namespace Tools.DotNETCommon
 
 						if (!bProcessInJob)
 						{
-							throw new Win32Exception(OriginalError);
+							throw new Win32ExceptionWithCode(OriginalError, "Unable to assign process to job object");
 						}
 					}
 
 					// Allow the thread to start running
 					if (ResumeThread(ProcessInfo.hThread) == -1)
 					{
-						throw new Win32Exception();
+						throw new Win32ExceptionWithCode("Unable to resume thread in child process");
 					}
 
 					// If we have any input text, write it to stdin now
diff --git a/Engine/Source/Programs/DotNETCommon/DotNETUtilities/Win32ExceptionWithCode.cs b/Engine/Source/Programs/DotNETCommon/DotNETUtilities/Win32ExceptionWithCode.cs
new file mode 100644
--- /dev/null
+++ b/Engine/Source/Programs/DotNETCommon/DotNETUtilities/Win32ExceptionWithCode.cs
@@ -0,0 +1,51 @@
+using System;
+using System.Collections.Generic;
+using System.ComponentModel;
+using System.Linq;
+using System.Text;
+using System.Threading.Tasks;
+
+namespace Tools.DotNETCommon
+{
+	/// <summary>
+	/// Wrapper for Win32Exception which includes the error code in the exception message
+	/// </summary>
+	class Win32ExceptionWithCode : Win32Exception
+	{
+		/// <summary>
+		/// Constructor
+		/// </summary>
+		/// <param name="Code">The Windows error code</param>
+		public Win32ExceptionWithCode(int Code)
+			: base(Code)
+		{
+		}
+
+		/// <summary>
+		/// Constructor
+		/// </summary>
+		/// <param name="Message">Message to display</param>
+		public Win32ExceptionWithCode(string Message)
+			: base(Message)
+		{
+		}
+
+		/// <summary>
+		/// Constructor
+		/// </summary>
+		/// <param name="Code">The Windows error code</param>
+		/// <param name="Message">Message to display</param>
+		public Win32ExceptionWithCode(int Code, string Message)
+			: base(Code, Message)
+		{
+		}
+
+		/// <summary>
+		/// Returns the exception message. Overriden to include the error code in the message.
+		/// </summary>
+		public override string Message
+		{
+			get { return String.Format("{0} (code 0x{1:X8})", base.Message, base.NativeErrorCode); }
+		}
+	}
+}
diff --git a/Engine/Source/Programs/UnrealBuildTool/System/Timeline.cs b/Engine/Source/Programs/UnrealBuildTool/System/Timeline.cs
--- a/Engine/Source/Programs/UnrealBuildTool/System/Timeline.cs
+++ b/Engine/Source/Programs/UnrealBuildTool/System/Timeline.cs
@@ -150,7 +150,7 @@ namespace UnrealBuildTool
 				Event Event = Events[EventIdx];
 
 				// Pop events off the stack
-				for (; OuterEvents.Count > 0; OuterEvents.RemoveAt(OuterEvents.Count - 1))
+				for (; OuterEvents.Count > 1; OuterEvents.RemoveAt(OuterEvents.Count - 1))
 				{
 					Event OuterEvent = OuterEvents.Last();
 					if (Event.StartTime < OuterEvent.FinishTime.Value)
diff --git a/Engine/Source/Programs/UnrealBuildTool/UnrealBuildTool.cs b/Engine/Source/Programs/UnrealBuildTool/UnrealBuildTool.cs
--- a/Engine/Source/Programs/UnrealBuildTool/UnrealBuildTool.cs
+++ b/Engine/Source/Programs/UnrealBuildTool/UnrealBuildTool.cs
@@ -533,14 +533,14 @@ namespace UnrealBuildTool
 			catch (BuildException Ex)
 			{
 				// BuildExceptions should have nicely formatted messages. We can log these directly.
-				Log.TraceError(Ex.Message.ToString());
+				Log.TraceError(ExceptionUtils.FormatException(Ex));
 				Log.TraceLog(ExceptionUtils.FormatExceptionDetails(Ex));
 				return (int)CompilationResult.OtherCompilationError;
 			}
 			catch (Exception Ex)
 			{
 				// Unhandled exception. 
-				Log.TraceError("Unhandled exception: {0}", Ex);
+				Log.TraceError("Unhandled exception: {0}", ExceptionUtils.FormatException(Ex));
 				Log.TraceLog(ExceptionUtils.FormatExceptionDetails(Ex));
 				return (int)CompilationResult.OtherCompilationError;
 			}
