#UE4 Added copyright

diff --git a/Engine/Source/Programs/DotNETCommon/DotNETUtilities/Win32ExceptionWithCode.cs b/Engine/Source/Programs/DotNETCommon/DotNETUtilities/Win32ExceptionWithCode.cs
--- a/Engine/Source/Programs/DotNETCommon/DotNETUtilities/Win32ExceptionWithCode.cs
+++ b/Engine/Source/Programs/DotNETCommon/DotNETUtilities/Win32ExceptionWithCode.cs
@@ -1,3 +1,5 @@
+// Copyright 1998-2019 Epic Games, Inc. All Rights Reserved.
+
 using System;
 using System.Collections.Generic;
 using System.ComponentModel;
