The "platforms as plugins" work exposed two awesome issues...

1) The CsProjectinfo class was not processing <Choose> nodes so ItemGroups within those nodes were not having their compile files parsed

2) The CsProjectInfo class only supported fully formed paths, not the wildcard syntax that csproj files can use.

This CL corrects both of these by adding the following -

* <Choose> nodes are now descended into and their elements are processed.
* The condition evalutor now supports the Exists statement (e.g. Exists(path))
* Wildcard include statements are expanded into a list of files. File wildcards (Foo/*.cs), path wildcards (Foo/*/*.cs), and greedy wildcards (Foo/**) are all supported.

As part of this work the Gauntlet.csproj file has been modified to use wildcards.

diff --git a/Engine/Source/Programs/AutomationTool/Gauntlet/Gauntlet.Automation.csproj b/Engine/Source/Programs/AutomationTool/Gauntlet/Gauntlet.Automation.csproj
--- a/Engine/Source/Programs/AutomationTool/Gauntlet/Gauntlet.Automation.csproj
+++ b/Engine/Source/Programs/AutomationTool/Gauntlet/Gauntlet.Automation.csproj
@@ -204,75 +204,8 @@
     <Content Include="SelfTest\TestData\LogParser\OrionPS4ClientLogWithPerf.txt" />
     <Content Include="SelfTest\TestData\LogParser\Win64FatalError.txt" />
   </ItemGroup>
-  <ItemGroup />
-  <ItemGroup />
-  <Choose>
-    <When Condition="Exists('Platform\Android\Gauntlet.TargetDeviceAndroid.cs')">
-      <ItemGroup>
-        <Compile Include="Platform\Android\Gauntlet.TargetDeviceAndroid.cs" />
-        <Compile Include="Platform\Android\Gauntlet.AndroidBuildSource.cs" />
-      </ItemGroup>
-    </When>
-  </Choose>
-  <Choose>
-    <When Condition="Exists('Platform\IOS\Gauntlet.TargetDeviceIOS.cs')">
-      <ItemGroup>
-        <Compile Include="Platform\IOS\Gauntlet.TargetDeviceIOS.cs" />
-        <Compile Include="Platform\IOS\Gauntlet.IOSBuildSource.cs" />
-      </ItemGroup>
-    </When>
-  </Choose>
-  <Choose>
-    <When Condition="Exists('Platform\Mac\Gauntlet.TargetDeviceMac.cs')">
-      <ItemGroup>
-        <Compile Include="Platform\Mac\Gauntlet.TargetDeviceMac.cs" />
-        <Compile Include="Platform\Mac\Gauntlet.MacBuildSource.cs" />
-        <Compile Include="Platform\Mac\Gauntlet.SelfTest.TestUnrealInstallAndRunMac.cs" />
-        <Compile Include="Platform\Mac\Gauntlet.SelfTest.TestTargetDeviceMac.cs" />
-      </ItemGroup>
-    </When>
-  </Choose>
-  <Choose>
-    <When Condition="Exists('Platform\PS4\Gauntlet.TargetDevicePS4.cs')">
-      <ItemGroup>
-        <Compile Include="Platform\PS4\Gauntlet.TargetDevicePS4.cs" />
-        <Compile Include="Platform\PS4\Gauntlet.PS4BuildSources.cs" />
-        <Compile Include="Platform\PS4\Gauntlet.SelfTest.TestUnrealInstallAndRunPS4.cs" />
-        <Compile Include="Platform\PS4\Gauntlet.SelfTest.TestTargetDevicePS4.cs" />
-        <Compile Include="Platform\PS4\Gauntlet.PS4PGO.cs" />
-        <Compile Include="Platform\PS4\Gauntlet.PS4Utils.cs" />
-      </ItemGroup>
-    </When>
-  </Choose>
-  <Choose>
-    <When Condition="Exists('Platform\Switch\Gauntlet.TargetDeviceSwitch.cs')">
-      <ItemGroup>
-        <Compile Include="Platform\Switch\Gauntlet.TargetDeviceSwitch.cs" />
-        <Compile Include="Platform\Switch\Gauntlet.SwitchBuildSource.cs" />
-        <Compile Include="Platform\Switch\Gauntlet.SwitchPGO.cs" />
-      </ItemGroup>
-    </When>
-  </Choose>
-  <Choose>
-    <When Condition="Exists('Platform\Windows\Gauntlet.TargetDeviceWindows.cs')">
-      <ItemGroup>
-        <Compile Include="Platform\Windows\Gauntlet.TargetDeviceWindows.cs" />
-        <Compile Include="Platform\Windows\Gauntlet.WindowsBuildSources.cs" />
-        <Compile Include="Platform\Windows\Gauntlet.WindowsPGO.cs" />
-        <Compile Include="Platform\Windows\Gauntlet.SelfTest.TestTargetDeviceWindows.cs" />
-      </ItemGroup>
-    </When>
-  </Choose>
-  <Choose>
-    <When Condition="Exists('Platform\XboxOne\Gauntlet.TargetDeviceXboxOne.cs')">
-      <ItemGroup>
-        <Compile Include="Platform\XboxOne\Gauntlet.TargetDeviceXboxOne.cs" />
-        <Compile Include="Platform\XboxOne\Gauntlet.XboxOneBuildSource.cs" />
-        <Compile Include="Platform\XboxOne\Gauntlet.SelfTest.TestUnrealInstallAndRunXboxOne.cs" />
-        <Compile Include="Platform\XboxOne\Gauntlet.SelfTest.TestTargetDeviceXboxOne.cs" />
-        <Compile Include="Platform\XboxOne\Gauntlet.XboxOnePGO.cs" />
-      </ItemGroup>
-    </When>
-  </Choose>
+  <ItemGroup>
+        <Compile Include="Platform\**\*.cs" />
+  </ItemGroup>
   <Import Project="$(MSBuildToolsPath)\Microsoft.CSharp.targets" />
 </Project>
\ No newline at end of file
diff --git a/Engine/Source/Programs/DotNETCommon/DotNETUtilities/CsProjectInfo.cs b/Engine/Source/Programs/DotNETCommon/DotNETUtilities/CsProjectInfo.cs
--- a/Engine/Source/Programs/DotNETCommon/DotNETUtilities/CsProjectInfo.cs
+++ b/Engine/Source/Programs/DotNETCommon/DotNETUtilities/CsProjectInfo.cs
@@ -2,6 +2,7 @@
 
 using System;
 using System.Collections.Generic;
+using System.IO;
 using System.Linq;
 using System.Text;
 using System.Threading.Tasks;
@@ -234,6 +235,41 @@ namespace Tools.DotNETCommon
 			}
 		}
 
+		/// <summary>
+		/// Parses supported elements from the children of the provided note. May recurse
+		/// for conditional elements.
+		/// </summary>
+		/// <param name="Node"></param>
+		/// <param name="ProjectInfo"></param>
+		static void ParseNode(XmlNode Node, CsProjectInfo ProjectInfo)
+		{
+			foreach (XmlElement Element in Node.ChildNodes.OfType<XmlElement>())
+			{
+				switch (Element.Name)
+				{
+					case "PropertyGroup":
+						if (EvaluateCondition(Element, ProjectInfo))
+						{
+							ParsePropertyGroup(Element, ProjectInfo);
+						}
+						break;
+					case "ItemGroup":
+						if (EvaluateCondition(Element, ProjectInfo))
+						{
+							ParseItemGroup(ProjectInfo.ProjectPath.Directory, Element, ProjectInfo);
+						}
+						break;
+					case "Choose":
+					case "When":
+						if (EvaluateCondition(Element, ProjectInfo))
+						{
+							ParseNode(Element, ProjectInfo);
+						}
+						break;
+				}
+			}
+		}
+
 		/// <summary>
 		/// Attempts to read project information for the given file.
 		/// </summary>
@@ -257,24 +293,9 @@ namespace Tools.DotNETCommon
 
 			// Parse the basic structure of the document, updating properties and recursing into other referenced projects as we go
 			CsProjectInfo ProjectInfo = new CsProjectInfo(Properties, File);
-			foreach (XmlElement Element in Document.DocumentElement.ChildNodes.OfType<XmlElement>())
-			{
-				switch (Element.Name)
-				{
-					case "PropertyGroup":
-						if (EvaluateCondition(Element, ProjectInfo.Properties))
-						{
-							ParsePropertyGroup(Element, ProjectInfo.Properties);
-						}
-						break;
-					case "ItemGroup":
-						if (EvaluateCondition(Element, ProjectInfo.Properties))
-						{
-							ParseItemGroup(File.Directory, Element, ProjectInfo);
-						}
-						break;
-				}
-			}
+
+			// Parse elements in the root node
+			ParseNode(Document.DocumentElement, ProjectInfo);
 
 			// Return the complete project
 			OutProjectInfo = ProjectInfo;
@@ -286,14 +307,14 @@ namespace Tools.DotNETCommon
 		/// </summary>
 		/// <param name="ParentElement">The parent 'PropertyGroup' element</param>
 		/// <param name="Properties">Dictionary mapping property names to values</param>
-		static void ParsePropertyGroup(XmlElement ParentElement, Dictionary<string, string> Properties)
+		static void ParsePropertyGroup(XmlElement ParentElement, CsProjectInfo ProjectInfo)
 		{
 			// We need to know the overridden output type and output path for the selected configuration.
 			foreach (XmlElement Element in ParentElement.ChildNodes.OfType<XmlElement>())
 			{
-				if (EvaluateCondition(Element, Properties))
+				if (EvaluateCondition(Element, ProjectInfo))
 				{
-					Properties[Element.Name] = ExpandProperties(Element.InnerText, Properties);
+					ProjectInfo.Properties[Element.Name] = ExpandProperties(Element.InnerText, ProjectInfo.Properties);
 				}
 			}
 		}
@@ -313,21 +334,21 @@ namespace Tools.DotNETCommon
 				{
 					case "Reference":
 						// Reference to an external assembly
-						if (EvaluateCondition(ItemElement, ProjectInfo.Properties))
+						if (EvaluateCondition(ItemElement, ProjectInfo))
 						{
 							ParseReference(BaseDirectory, ItemElement, ProjectInfo.References);
 						}
 						break;
 					case "ProjectReference":
 						// Reference to another project
-						if (EvaluateCondition(ItemElement, ProjectInfo.Properties))
+						if (EvaluateCondition(ItemElement, ProjectInfo))
 						{
 							ParseProjectReference(BaseDirectory, ItemElement, ProjectInfo.ProjectReferences);
 						}
 						break;
 					case "Compile":
 						// Reference to another project
-						if (EvaluateCondition(ItemElement, ProjectInfo.Properties))
+						if (EvaluateCondition(ItemElement, ProjectInfo))
 						{
 							ParseCompileReference(BaseDirectory, ItemElement, ProjectInfo.CompileReferences);
 						}
@@ -335,7 +356,7 @@ namespace Tools.DotNETCommon
 					case "Content":
 					case "None":
 						// Reference to another project
-						if (EvaluateCondition(ItemElement, ProjectInfo.Properties))
+						if (EvaluateCondition(ItemElement, ProjectInfo))
 						{
 							ParseContent(BaseDirectory, ItemElement, ProjectInfo.ContentReferences);
 						}
@@ -383,6 +404,100 @@ namespace Tools.DotNETCommon
 			}
 		}
 
+		/// <summary>
+		/// Finds all files in the provided path, which may be a csproj wildcard specification.
+		/// E.g. The following are all valid
+		/// Foo/Bar/Item.cs
+		/// Foo/Bar/*.cs
+		/// Foo/*/Item.cs
+		/// Foo/*/*.cs
+		/// Foo/**
+		/// (the last means include all files under the path).
+		/// </summary>
+		/// <param name="InPath">Path specifier to process</param>
+		/// <returns></returns>
+		static IEnumerable<FileReference> FindMatchingFiles(FileReference InPath)
+		{		
+			List<FileReference> FoundFiles = new List<FileReference>();		
+
+			// split off the drive root
+			string DriveRoot = Path.GetPathRoot(InPath.FullName);
+
+			// break the rest of the path into components
+			string[] PathComponents = InPath.FullName.Substring(DriveRoot.Length).Split(new char[] { Path.DirectorySeparatorChar, Path.AltDirectorySeparatorChar });
+
+			// Process all the components recursively
+			ProcessPathComponents(new DirectoryReference(DriveRoot), PathComponents, FoundFiles);
+
+			/// recursive function that will append RemainingComponents one by one to ExistingPath, expanding wildcards as necessary. The complete
+			/// list of files that match the complete path is returned out OutFoundFiles
+			void ProcessPathComponents(DirectoryReference ExistingPath, IEnumerable<string> RemainingComponents, List<FileReference> OutFoundFiles)
+			{
+				if (!RemainingComponents.Any())
+				{
+					return;
+				}
+
+				// take a look at the first component
+				string CurrentComponent = RemainingComponents.First();
+				RemainingComponents = RemainingComponents.Skip(1);
+
+				// If no other components then this is either a file pattern or a greedy pattern
+				if (!RemainingComponents.Any())
+				{
+					// ** means include all files under this tree, so enumerate them all
+					if (CurrentComponent.Contains("**"))
+					{
+						OutFoundFiles.AddRange(DirectoryReference.EnumerateFiles(ExistingPath, "*", SearchOption.AllDirectories));
+					}
+					else
+					{
+						// easy, a regular path with a file that may or may not be a wildcard
+						OutFoundFiles.AddRange(DirectoryReference.EnumerateFiles(ExistingPath, CurrentComponent));
+					}
+				}
+				else
+				{
+					// new component contains a wildcard, and based on the above we know there are more entries so find 
+					// matching directories
+					if (CurrentComponent.Contains("*"))
+					{
+						// ** means all directories, no matter how deep
+						SearchOption Option = CurrentComponent == "**" ? SearchOption.AllDirectories : SearchOption.TopDirectoryOnly;
+
+						IEnumerable<DirectoryReference> Directories = DirectoryReference.EnumerateDirectories(ExistingPath, CurrentComponent, Option);
+
+						// if we searched all directories regardless of depth, the rest of the components other than the last (file) are irrelevant
+						if (Option == SearchOption.AllDirectories)
+						{
+							RemainingComponents = new[] { RemainingComponents.Last() };
+						}
+
+						foreach (DirectoryReference Dir in Directories)
+						{
+							ProcessPathComponents(Dir, RemainingComponents, OutFoundFiles);
+						}
+					}
+					else
+					{
+						// add this component to our path and recurse.
+						ExistingPath = DirectoryReference.Combine(ExistingPath, CurrentComponent);
+
+						// but... we can just take all the next components that don't have wildcards in them instead of recursing
+						// into each one!
+						IEnumerable<string> NonWildCardComponents = RemainingComponents.TakeWhile(C => !C.Contains("*"));
+						RemainingComponents = RemainingComponents.Skip(NonWildCardComponents.Count());
+
+						ExistingPath = DirectoryReference.Combine(ExistingPath, NonWildCardComponents.ToArray());
+
+						ProcessPathComponents(ExistingPath, RemainingComponents, OutFoundFiles);
+					}
+				}
+			}
+
+			return FoundFiles;
+		}
+
 		/// <summary>
 		/// Parses a project reference from a given 'ProjectReference' element
 		/// </summary>
@@ -395,7 +510,15 @@ namespace Tools.DotNETCommon
 			if (!String.IsNullOrEmpty(IncludePath))
 			{
 				FileReference SourceFile = FileReference.Combine(BaseDirectory, IncludePath);
-				CompileReferences.Add(SourceFile);
+
+				if (SourceFile.FullName.Contains("*"))
+				{
+					CompileReferences.AddRange(FindMatchingFiles(SourceFile));
+				}
+				else
+				{
+					CompileReferences.Add(SourceFile);
+				}
 			}
 		}
 
@@ -471,7 +594,7 @@ namespace Tools.DotNETCommon
 		/// <param name="Element">The XML element to check</param>
 		/// <param name="Properties">Dictionary mapping from property names to values.</param>
 		/// <returns></returns>
-		static bool EvaluateCondition(XmlElement Element, Dictionary<string, string> Properties)
+		static bool EvaluateCondition(XmlElement Element, CsProjectInfo ProjectInfo)
 		{
 			// Read the condition attribute. If it's not present, assume it evaluates to true.
 			string Condition = Element.GetAttribute("Condition");
@@ -481,7 +604,7 @@ namespace Tools.DotNETCommon
 			}
 
 			// Expand all the properties
-			Condition = ExpandProperties(Condition, Properties);
+			Condition = ExpandProperties(Condition, ProjectInfo.Properties);
 
 			// Parse literal true/false values
 			bool OutResult;
@@ -493,8 +616,28 @@ namespace Tools.DotNETCommon
 			// Tokenize the condition
 			string[] Tokens = Tokenize(Condition);
 
+			char[] TokenQuotes = new[] { '\'', '(', ')', '{', '}', '[', ']' };
+			
 			// Try to evaluate it. We only support a very limited class of condition expressions at the moment, but it's enough to parse standard projects
 			bool bResult;
+
+			// Handle Exists('Platform\Windows\Gauntlet.TargetDeviceWindows.cs')
+			if (Tokens[0] == "Exists")
+			{
+				// remove all quotes, apostrophes etc that are either tokens or wrap tokens (The Tokenize() function is a bit suspect).
+				string[] Arguments = Tokens.Select(S => S.Trim(TokenQuotes)).Where(S => S.Length > 0).ToArray();
+
+				if (Tokens.Length > 1)
+				{
+					FileSystemReference Dependency = DirectoryReference.Combine(ProjectInfo.ProjectPath.Directory, Arguments[1]);
+
+					if (File.Exists(Dependency.FullName) || Directory.Exists(Dependency.FullName))
+					{
+						return true;
+					}
+				}
+			}
+
 			if (Tokens.Length == 3 && Tokens[0].StartsWith("'") && Tokens[1] == "==" && Tokens[2].StartsWith("'"))
 			{
 				bResult = String.Compare(Tokens[0], Tokens[2], StringComparison.InvariantCultureIgnoreCase) == 0;
