Fixing odd issue on builders with nested function

diff --git a/Engine/Source/Programs/DotNETCommon/DotNETUtilities/CsProjectInfo.cs b/Engine/Source/Programs/DotNETCommon/DotNETUtilities/CsProjectInfo.cs
--- a/Engine/Source/Programs/DotNETCommon/DotNETUtilities/CsProjectInfo.cs
+++ b/Engine/Source/Programs/DotNETCommon/DotNETUtilities/CsProjectInfo.cs
@@ -404,6 +404,73 @@ namespace Tools.DotNETCommon
 			}
 		}
 
+
+		/// recursive helper used by the function below that will append RemainingComponents one by one to ExistingPath, 
+		/// expanding wildcards as necessary. The complete list of files that match the complete path is returned out OutFoundFiles
+		static void ProcessPathComponents(DirectoryReference ExistingPath, IEnumerable<string> RemainingComponents, List<FileReference> OutFoundFiles)
+		{
+			if (!RemainingComponents.Any())
+			{
+				return;
+			}
+
+			// take a look at the first component
+			string CurrentComponent = RemainingComponents.First();
+			RemainingComponents = RemainingComponents.Skip(1);
+
+			// If no other components then this is either a file pattern or a greedy pattern
+			if (!RemainingComponents.Any())
+			{
+				// ** means include all files under this tree, so enumerate them all
+				if (CurrentComponent.Contains("**"))
+				{
+					OutFoundFiles.AddRange(DirectoryReference.EnumerateFiles(ExistingPath, "*", SearchOption.AllDirectories));
+				}
+				else
+				{
+					// easy, a regular path with a file that may or may not be a wildcard
+					OutFoundFiles.AddRange(DirectoryReference.EnumerateFiles(ExistingPath, CurrentComponent));
+				}
+			}
+			else
+			{
+				// new component contains a wildcard, and based on the above we know there are more entries so find 
+				// matching directories
+				if (CurrentComponent.Contains("*"))
+				{
+					// ** means all directories, no matter how deep
+					SearchOption Option = CurrentComponent == "**" ? SearchOption.AllDirectories : SearchOption.TopDirectoryOnly;
+
+					IEnumerable<DirectoryReference> Directories = DirectoryReference.EnumerateDirectories(ExistingPath, CurrentComponent, Option);
+
+					// if we searched all directories regardless of depth, the rest of the components other than the last (file) are irrelevant
+					if (Option == SearchOption.AllDirectories)
+					{
+						RemainingComponents = new[] { RemainingComponents.Last() };
+					}
+
+					foreach (DirectoryReference Dir in Directories)
+					{
+						ProcessPathComponents(Dir, RemainingComponents, OutFoundFiles);
+					}
+				}
+				else
+				{
+					// add this component to our path and recurse.
+					ExistingPath = DirectoryReference.Combine(ExistingPath, CurrentComponent);
+
+					// but... we can just take all the next components that don't have wildcards in them instead of recursing
+					// into each one!
+					IEnumerable<string> NonWildCardComponents = RemainingComponents.TakeWhile(C => !C.Contains("*"));
+					RemainingComponents = RemainingComponents.Skip(NonWildCardComponents.Count());
+
+					ExistingPath = DirectoryReference.Combine(ExistingPath, NonWildCardComponents.ToArray());
+
+					ProcessPathComponents(ExistingPath, RemainingComponents, OutFoundFiles);
+				}
+			}
+		}
+
 		/// <summary>
 		/// Finds all files in the provided path, which may be a csproj wildcard specification.
 		/// E.g. The following are all valid
@@ -429,72 +496,6 @@ namespace Tools.DotNETCommon
 			// Process all the components recursively
 			ProcessPathComponents(new DirectoryReference(DriveRoot), PathComponents, FoundFiles);
 
-			/// recursive function that will append RemainingComponents one by one to ExistingPath, expanding wildcards as necessary. The complete
-			/// list of files that match the complete path is returned out OutFoundFiles
-			void ProcessPathComponents(DirectoryReference ExistingPath, IEnumerable<string> RemainingComponents, List<FileReference> OutFoundFiles)
-			{
-				if (!RemainingComponents.Any())
-				{
-					return;
-				}
-
-				// take a look at the first component
-				string CurrentComponent = RemainingComponents.First();
-				RemainingComponents = RemainingComponents.Skip(1);
-
-				// If no other components then this is either a file pattern or a greedy pattern
-				if (!RemainingComponents.Any())
-				{
-					// ** means include all files under this tree, so enumerate them all
-					if (CurrentComponent.Contains("**"))
-					{
-						OutFoundFiles.AddRange(DirectoryReference.EnumerateFiles(ExistingPath, "*", SearchOption.AllDirectories));
-					}
-					else
-					{
-						// easy, a regular path with a file that may or may not be a wildcard
-						OutFoundFiles.AddRange(DirectoryReference.EnumerateFiles(ExistingPath, CurrentComponent));
-					}
-				}
-				else
-				{
-					// new component contains a wildcard, and based on the above we know there are more entries so find 
-					// matching directories
-					if (CurrentComponent.Contains("*"))
-					{
-						// ** means all directories, no matter how deep
-						SearchOption Option = CurrentComponent == "**" ? SearchOption.AllDirectories : SearchOption.TopDirectoryOnly;
-
-						IEnumerable<DirectoryReference> Directories = DirectoryReference.EnumerateDirectories(ExistingPath, CurrentComponent, Option);
-
-						// if we searched all directories regardless of depth, the rest of the components other than the last (file) are irrelevant
-						if (Option == SearchOption.AllDirectories)
-						{
-							RemainingComponents = new[] { RemainingComponents.Last() };
-						}
-
-						foreach (DirectoryReference Dir in Directories)
-						{
-							ProcessPathComponents(Dir, RemainingComponents, OutFoundFiles);
-						}
-					}
-					else
-					{
-						// add this component to our path and recurse.
-						ExistingPath = DirectoryReference.Combine(ExistingPath, CurrentComponent);
-
-						// but... we can just take all the next components that don't have wildcards in them instead of recursing
-						// into each one!
-						IEnumerable<string> NonWildCardComponents = RemainingComponents.TakeWhile(C => !C.Contains("*"));
-						RemainingComponents = RemainingComponents.Skip(NonWildCardComponents.Count());
-
-						ExistingPath = DirectoryReference.Combine(ExistingPath, NonWildCardComponents.ToArray());
-
-						ProcessPathComponents(ExistingPath, RemainingComponents, OutFoundFiles);
-					}
-				}
-			}
-
 			return FoundFiles;
 		}
 
