Refactor how we find child process IDs.

#rnx

diff --git a/Engine/Source/Programs/DotNETCommon/DotNETUtilities/ProcessUtils.cs b/Engine/Source/Programs/DotNETCommon/DotNETUtilities/ProcessUtils.cs
--- a/Engine/Source/Programs/DotNETCommon/DotNETUtilities/ProcessUtils.cs
+++ b/Engine/Source/Programs/DotNETCommon/DotNETUtilities/ProcessUtils.cs
@@ -104,26 +104,21 @@ namespace Tools.DotNETCommon
 				}
 
 				// Find any process ids which are a descendant from the root process it
-				List<uint> ChildProcessIds = new List<uint>();
-				foreach (KeyValuePair<uint, uint> Pair in ProcessIdToParentProcessId)
+				HashSet<uint> ChildProcessIds = new HashSet<uint>();
+				for (bool bContinueLoop = true; bContinueLoop; )
 				{
-					uint ParentProcessId = Pair.Value;
-					for (; ; )
+					bContinueLoop = false;
+					foreach (KeyValuePair<uint, uint> Pair in ProcessIdToParentProcessId)
 					{
-						if (ParentProcessId == RootProcessId)
+						if (Pair.Value == RootProcessId || ChildProcessIds.Contains(Pair.Value))
 						{
-							ChildProcessIds.Add(Pair.Key);
-							break;
-						}
-						if (!ProcessIdToParentProcessId.TryGetValue(ParentProcessId, out ParentProcessId))
-						{
-							break;
+							bContinueLoop |= ChildProcessIds.Add(Pair.Key);
 						}
 					}
 				}
 
 				// If there's nothing running, we can quit
-				if (ChildProcessIds.Count == 0)
+				if (!ChildProcessIds.Any())
 				{
 					break;
 				}
