Merging Engine/Source/Programs/DotNETCommon/MetaData.cs from Dev-EngineMerge to Main.

#rnx

diff --git a/Engine/Source/Programs/DotNETCommon/MetaData.cs b/Engine/Source/Programs/DotNETCommon/MetaData.cs
--- a/Engine/Source/Programs/DotNETCommon/MetaData.cs
+++ b/Engine/Source/Programs/DotNETCommon/MetaData.cs
@@ -1,4 +1,4 @@
-// Copyright 1998-2019 Epic Games, Inc. All Rights Reserved.
+// Copyright Epic Games, Inc. All Rights Reserved.
 
 using System.Reflection;
 using System.Runtime.InteropServices;
@@ -7,7 +7,7 @@ using System.Resources;
 // These are the assembly properties for all tools
 [assembly: AssemblyCompany( "Epic Games, Inc." )]
 [assembly: AssemblyProduct( "UE4" )]
-[assembly: AssemblyCopyright( "Copyright 1998-2019 Epic Games, Inc. All Rights Reserved." )]
+[assembly: AssemblyCopyright( "Copyright Epic Games, Inc. All Rights Reserved." )]
 [assembly: AssemblyTrademark( "" )]
 
 // Use a neutral culture to avoid some localisation issues
