Fixing an edge case crash when trying to word wrap a very long "word" and the current sentence is already full.

#rnx

diff --git a/Engine/Source/Programs/DotNETCommon/DotNETUtilities/StringUtils.cs b/Engine/Source/Programs/DotNETCommon/DotNETUtilities/StringUtils.cs
--- a/Engine/Source/Programs/DotNETCommon/DotNETUtilities/StringUtils.cs
+++ b/Engine/Source/Programs/DotNETCommon/DotNETUtilities/StringUtils.cs
@@ -51,6 +51,14 @@ namespace Tools.DotNETCommon
 				// if this is a very large word, split it
 				if (Word.Length > MaxWidth)
 				{
+					// If the current sentence is ready to be written, do that.
+					if (CurrentSentence.Length >= MaxWidth)
+					{
+						// next line and reset sentence
+						WrappedWords.Add(CurrentSentence);
+						CurrentSentence = string.Empty;
+					}
+
 					// Top up the current line
 					WrappedWords.Add(CurrentSentence + Word.Substring(0, MaxWidth - CurrentSentence.Length));
 
