PR #6412: FIX: Added support for FileReference to fix serialisation of RemoteMac.SshPrivateKey (Contributed by [REDACTED_PERSON])

diff --git a/Engine/Source/Programs/DotNETCommon/DotNETUtilities/BinaryReaderExtensions.cs b/Engine/Source/Programs/DotNETCommon/DotNETUtilities/BinaryReaderExtensions.cs
--- a/Engine/Source/Programs/DotNETCommon/DotNETUtilities/BinaryReaderExtensions.cs
+++ b/Engine/Source/Programs/DotNETCommon/DotNETUtilities/BinaryReaderExtensions.cs
@@ -193,6 +193,10 @@ namespace Tools.DotNETCommon
 				int Value = Reader.ReadInt32();
 				return (Value == -1)? (bool?)null : (Value == 0)? (bool?)false : (bool?)true;
 			}
+			else if (ObjectType == typeof(FileReference))
+			{
+				return Reader.ReadFileReference();
+			}
 			else if(ObjectType.IsEnum)
 			{
 				return Enum.ToObject(ObjectType, Reader.ReadInt32());
diff --git a/Engine/Source/Programs/DotNETCommon/DotNETUtilities/BinaryWriterExtensions.cs b/Engine/Source/Programs/DotNETCommon/DotNETUtilities/BinaryWriterExtensions.cs
--- a/Engine/Source/Programs/DotNETCommon/DotNETUtilities/BinaryWriterExtensions.cs
+++ b/Engine/Source/Programs/DotNETCommon/DotNETUtilities/BinaryWriterExtensions.cs
@@ -210,6 +210,10 @@ namespace Tools.DotNETCommon
 				bool? NullableValue = (bool?)Value;
 				Writer.Write(NullableValue.HasValue? NullableValue.Value? 1 : 0 : -1);
 			}
+			else if (FieldType == typeof(FileReference))
+			{
+				Writer.Write((FileReference)Value);
+			}
 			else
 			{
 				throw new Exception(String.Format("Unsupported type '{0}' for binary serialization", FieldType.Name));
diff --git a/Engine/Source/Programs/UnrealBuildTool/System/XmlConfig.cs b/Engine/Source/Programs/UnrealBuildTool/System/XmlConfig.cs
--- a/Engine/Source/Programs/UnrealBuildTool/System/XmlConfig.cs
+++ b/Engine/Source/Programs/UnrealBuildTool/System/XmlConfig.cs
@@ -624,6 +624,10 @@ namespace UnrealBuildTool
 			{
 				return Enum.Parse(FieldType, TrimmedText);
 			}
+			else if (FieldType == typeof(FileReference))
+			{
+				return FileReference.FromString(Text);
+			}
 			else
 			{
 				throw new Exception(String.Format("Unsupported config type '{0}'", FieldType.Name));
