Allow specifying a timespan for ThreadPoolWorkQueue.Wait().

diff --git a/Engine/Source/Programs/DotNETCommon/DotNETUtilities/ThreadPoolWorkQueue.cs b/Engine/Source/Programs/DotNETCommon/DotNETUtilities/ThreadPoolWorkQueue.cs
--- a/Engine/Source/Programs/DotNETCommon/DotNETUtilities/ThreadPoolWorkQueue.cs
+++ b/Engine/Source/Programs/DotNETCommon/DotNETUtilities/ThreadPoolWorkQueue.cs
@@ -132,8 +132,18 @@ namespace Tools.DotNETCommon
 		/// <returns>True if the queue completed, false if the timeout elapsed</returns>
 		public bool Wait(int MillisecondsTimeout)
 		{
-			bool bResult = EmptyEvent.WaitOne(MillisecondsTimeout);
-			if(bResult)
+			return Wait(TimeSpan.FromMilliseconds(MillisecondsTimeout));
+		}
+
+		/// <summary>
+		/// Waits for all queued tasks to finish, or the timeout to elapse
+		/// </summary>
+		/// <param name="Timeout">Maximum time to wait</param>
+		/// <returns>True if the queue completed, false if the timeout elapsed</returns>
+		public bool Wait(TimeSpan Timeout)
+		{
+			bool bResult = EmptyEvent.WaitOne(Timeout);
+			if (bResult)
 			{
 				RethrowExceptions();
 			}
