- Fixed paths to known ThirdParty assemblies

diff --git a/Engine/Source/Programs/DotNETCommon/DotNETUtilities/AssemblyUtils.cs b/Engine/Source/Programs/DotNETCommon/DotNETUtilities/AssemblyUtils.cs
--- a/Engine/Source/Programs/DotNETCommon/DotNETUtilities/AssemblyUtils.cs
+++ b/Engine/Source/Programs/DotNETCommon/DotNETUtilities/AssemblyUtils.cs
@@ -44,16 +44,15 @@ namespace Tools.DotNETCommon
 			    // Name is fully qualified assembly definition - e.g. "p4dn, Version=1.0.0.0, Culture=neutral, PublicKeyToken=ff968dc1933aba6f"
 			    string AssemblyName = args.Name.Split(',')[0];
 
-                // Look for known assembly names we check into Binaries/DotNET/. Return null if we can't find it.
-                return (
-                    from KnownAssemblyName in new[] { "SwarmAgent.exe", "Ionic.Zip.Reduced.dll", "Newtonsoft/NewtonSoft.Json.dll" }
+				return (
+                    from KnownAssemblyName in new[] { "SwarmAgent.exe", "../ThirdParty/Ionic/Ionic.Zip.Reduced.dll", "../ThirdParty/Newtonsoft/NewtonSoft.Json.dll" }
                     where AssemblyName.Equals(Path.GetFileNameWithoutExtension(KnownAssemblyName), StringComparison.InvariantCultureIgnoreCase)
                     let ResolvedAssemblyFilename = Path.Combine(PathToBinariesDotNET, KnownAssemblyName)
                     // check if the file exists first. If we just try to load it, we correctly throw an exception, but it's a generic
                     // FileNotFoundException, which is not informative. Better to return null.
                     select File.Exists(ResolvedAssemblyFilename) ? Assembly.LoadFile(ResolvedAssemblyFilename) : null
                     ).FirstOrDefault();
-            };
+			};
         }
 
 		/// <summary>
