- Moved UEBuildPlatformSDK into a separate DLL, BuildUtilities (similar to DotNETUtilities), allowing for better sharing of code between UBT and UAT (Turnkey in particular)
- Refactored SDK versioning into more formalized APIs for all platforms. Currently UBT will not error with bad versions, but that is coming.
- Turnkey updated to use new BuildUtilities

diff --git a/Engine/Plugins/Lumin/MLSDK/Source/MLSDK/MLSDK.Build.cs b/Engine/Plugins/Lumin/MLSDK/Source/MLSDK/MLSDK.Build.cs
--- a/Engine/Plugins/Lumin/MLSDK/Source/MLSDK/MLSDK.Build.cs
+++ b/Engine/Plugins/Lumin/MLSDK/Source/MLSDK/MLSDK.Build.cs
@@ -15,11 +15,11 @@ public class MLSDK : ModuleRules
 
 		Type = ModuleType.External;
 
-		LuminSDKVersionHelper SDKVersionHelper = new LuminSDKVersionHelper();
+		UEBuildPlatformSDK SDK = UEBuildPlatformSDK.GetSDKForPlatform(UnrealTargetPlatform.Lumin.ToString());
 		string MLSDKPath = System.Environment.GetEnvironmentVariable("MLSDK");
 		bool bIsMLSDKInstalled = false;
 		bool bIsSupportedTargetPlatformAndType = ((Target.Platform == UnrealTargetPlatform.Lumin) || (Target.Type == TargetRules.TargetType.Editor));
-		if (SDKVersionHelper.HasAnySDK() && !string.IsNullOrEmpty(MLSDKPath))
+		if (SDK.IsVersionValid(SDK.GetInstalledVersion(), bForAutoSDK:false) && !string.IsNullOrEmpty(MLSDKPath))
 		{
 			string IncludePath = Path.Combine(MLSDKPath, "include");
 			string LibraryPath = Path.Combine(MLSDKPath, "lib");
diff --git a/Engine/Source/Programs/AutomationTool/AutomationUtils/Platform.cs b/Engine/Source/Programs/AutomationTool/AutomationUtils/Platform.cs
--- a/Engine/Source/Programs/AutomationTool/AutomationUtils/Platform.cs
+++ b/Engine/Source/Programs/AutomationTool/AutomationUtils/Platform.cs
@@ -158,20 +158,6 @@ namespace AutomationTool
 
 		#region Turnkey
 
-		public virtual string GetInstalledSdk()
-		{
-			return "";
-		}
-
-		public virtual string GetAllowedSdks()
-		{
-			return "";
-		}
-		public virtual string GetAutoSdkPlatformName()
-		{
-			return IniPlatformType.ToString();
-		}
-
 		public virtual DeviceInfo[] GetDevices()
 		{
 			return null;
@@ -199,11 +185,6 @@ namespace AutomationTool
 		{
 			return true;
 		}
-		public virtual bool TryConvertVersionToInt(string StringValue, out UInt64 IntValue)
-		{
-			IntValue = 0;
-			return false;
-		}
 
 		#endregion
 
diff --git a/Engine/Source/Programs/AutomationTool/AutomationUtils/UE4Build.cs b/Engine/Source/Programs/AutomationTool/AutomationUtils/UE4Build.cs
--- a/Engine/Source/Programs/AutomationTool/AutomationUtils/UE4Build.cs
+++ b/Engine/Source/Programs/AutomationTool/AutomationUtils/UE4Build.cs
@@ -1325,6 +1325,7 @@ namespace AutomationTool
 						"UnrealBuildTool.exe.config",
 						"AutomationUtils.Automation.dll",
 						"DotNETUtilities.dll",
+						"BuildUtilities.dll",
 						"MobileDeviceInterface.dll"
 					});
 
diff --git a/Engine/Source/Programs/AutomationTool/BuildGraph/BuildGraph.cs b/Engine/Source/Programs/AutomationTool/BuildGraph/BuildGraph.cs
--- a/Engine/Source/Programs/AutomationTool/BuildGraph/BuildGraph.cs
+++ b/Engine/Source/Programs/AutomationTool/BuildGraph/BuildGraph.cs
@@ -507,7 +507,7 @@ namespace AutomationTool
 				}
 				catch (ReflectionTypeLoadException ex)
 				{
-					LogWarning("Exception {0} while trying to get types from assembly {1}", ex, LoadedAssembly);
+					LogWarning("Exception {0} while trying to get types from assembly {1}. LoaderExceptions: {2}", ex, LoadedAssembly, string.Join("\n", ex.LoaderExceptions.Select(x => x.Message)));
 					continue;
 				}
 
diff --git a/Engine/Source/Programs/AutomationTool/IOS/IOSPlatform.Automation.cs b/Engine/Source/Programs/AutomationTool/IOS/IOSPlatform.Automation.cs
--- a/Engine/Source/Programs/AutomationTool/IOS/IOSPlatform.Automation.cs
+++ b/Engine/Source/Programs/AutomationTool/IOS/IOSPlatform.Automation.cs
@@ -169,22 +169,6 @@ public class IOSPlatform : Platform
 		SDKName = (TargetPlatform == UnrealTargetPlatform.TVOS) ? "appletvos" : "iphoneos";
 	}
 
-	public override string GetAllowedSdks()
-	{
-		return "11.5";
-	}
-	public override string GetInstalledSdk()
-	{
-		if (HostPlatform.Current.HostEditorPlatform != UnrealTargetPlatform.Mac)
-		{
-			return base.GetInstalledSdk();
-		}
-
-		// @todo turnkey: get installed Xcode version
-		string Output = UnrealBuildTool.Utils.RunLocalProcessAndReturnStdOut("sh", "-c 'xcodebuild -version'");
-		Match Result = Regex.Match(Output, @"Xcode (\S*)");
-		return Result.Success ? Result.Groups[1].Value : "";
-	}
 	public override string GetAllowedSoftwareVersions()
 	{
 		return @"regex:^13\.5.*";
diff --git a/Engine/Source/Programs/AutomationTool/Turnkey/Commands/InstallSdk.cs b/Engine/Source/Programs/AutomationTool/Turnkey/Commands/InstallSdk.cs
--- a/Engine/Source/Programs/AutomationTool/Turnkey/Commands/InstallSdk.cs
+++ b/Engine/Source/Programs/AutomationTool/Turnkey/Commands/InstallSdk.cs
@@ -83,21 +83,22 @@ namespace Turnkey.Commands
 
 				// cache the automation platform object
 				AutomationTool.Platform AutomationPlatform = AutomationTool.Platform.GetPlatform(Platform);
-
-// 				if (bAllowAutoSdk)
-// 				{
-// 					// first, attempt AutoSdk
-// 					SdkInfo.LocalAvailability LocalState = SdkInfo.GetLocalAvailability(AutomationPlatform);
-// 					bool bWasAutoSdkSetup;
-// 					
-// 					SdkInfo.ConditionalSetupAutoSdk(Platform, ref LocalState, out bWasAutoSdkSetup, bUnattended: false);
-// 
-// 					// if we got an AutoSdk for this platform, then we don't need to continue
-// 					if (bWasAutoSdkSetup)
-// 					{
-// 						continue;
-// 					}
-// 				}
+				UEBuildPlatformSDK SDK = UEBuildPlatformSDK.GetSDKForPlatform(Platform.ToString());
+
+				// 				if (bAllowAutoSdk)
+				// 				{
+				// 					// first, attempt AutoSdk
+				// 					SdkInfo.LocalAvailability LocalState = SdkInfo.GetLocalAvailability(AutomationPlatform);
+				// 					bool bWasAutoSdkSetup;
+				// 					
+				// 					SdkInfo.ConditionalSetupAutoSdk(Platform, ref LocalState, out bWasAutoSdkSetup, bUnattended: false);
+				// 
+				// 					// if we got an AutoSdk for this platform, then we don't need to continue
+				// 					if (bWasAutoSdkSetup)
+				// 					{
+				// 						continue;
+				// 					}
+				// 				}
 
 
 				// filter the Sdks if a platform was given
@@ -108,14 +109,14 @@ namespace Turnkey.Commands
 				Sdks = Sdks.FindAll(x => x.Type != SdkInfo.SdkType.Flash || (AutomationPlatform.GetDevices() != null && AutomationPlatform.GetDevices().Length > 0));
 				bool bStrippedDevices = Sdks.Count != SdkCount;
 
-				// we don't need to do AutoSdks, we already attempted one above if we wanted to
-				Sdks = Sdks.FindAll(x => x.Type != SdkInfo.SdkType.AutoSdk);
+// 				// we don't need to do AutoSdks, we already attempted one above if we wanted to
+// 				Sdks = Sdks.FindAll(x => x.Type != SdkInfo.SdkType.AutoSdk);
 
 
 				// strip out Sdks where there is no Sdk installed yet
 				if (bUpdateOnly)
 				{
-					Sdks = Sdks.FindAll(x => !string.IsNullOrEmpty(AutomationPlatform.GetInstalledSdk()));
+					Sdks = Sdks.FindAll(x => !string.IsNullOrEmpty(SDK.GetInstalledVersion()));
 				}
 
 				// strip out Sdks not in the allowed range of the platform
@@ -205,7 +206,7 @@ namespace Turnkey.Commands
 					List<string> Options = new List<string>();
 					foreach (SdkInfo Sdk in Sdks)
 					{
-						string Current = AutomationPlatform.GetInstalledSdk();
+						string Current = SDK.GetInstalledVersion();
 						if (Sdk.Type == SdkInfo.SdkType.Flash)
 						{
 							// look for default device, or matching device [put this in a function!]
diff --git a/Engine/Source/Programs/AutomationTool/Turnkey/Commands/ListPlatforms.cs b/Engine/Source/Programs/AutomationTool/Turnkey/Commands/ListPlatforms.cs
--- a/Engine/Source/Programs/AutomationTool/Turnkey/Commands/ListPlatforms.cs
+++ b/Engine/Source/Programs/AutomationTool/Turnkey/Commands/ListPlatforms.cs
@@ -4,7 +4,7 @@ using System;
 using System.Collections.Generic;
 using UnrealBuildTool;
 using AutomationTool;
-
+using Tools.DotNETCommon;
 
 namespace Turnkey.Commands
 {
@@ -14,18 +14,38 @@ namespace Turnkey.Commands
 		{
 			TurnkeyUtils.Log("");
 			TurnkeyUtils.Log("Valid Platforms:");
+			string PlatformString = TurnkeyUtils.ParseParamValue("Platform", null, CommandOptions);
+
 			foreach (UnrealTargetPlatform TargetPlatform in UnrealTargetPlatform.GetValidPlatforms())
 			{
+				// HACK UNTIUL WIN32 IS GONE COMPLETELY
+				if (TargetPlatform == UnrealTargetPlatform.Win32)
+				{
+					continue;
+				}
+
+				if (PlatformString != null && UnrealTargetPlatform.Parse(PlatformString) != TargetPlatform)
+				{
+					continue;
+				}
+
 				Platform Platform = Platform.Platforms[new TargetPlatformDescriptor(TargetPlatform)];
+				UEBuildPlatformSDK SDK = UEBuildPlatformSDK.GetSDKForPlatform(TargetPlatform.ToString());
 
-				string InstalledSdk = Platform.GetInstalledSdk();
-				string AllowedSdks = Platform.GetAllowedSdks();
+				string ManualSDKVersion, AutoSDKVersion;
+				string MinAllowedVersion, MaxAllowedVersion;
+				SDK.GetInstalledVersions(out ManualSDKVersion, out AutoSDKVersion);
+				SDK.GetValidVersionRange(out MinAllowedVersion, out MaxAllowedVersion);
+				
 				string AllowedSoftware = Platform.GetAllowedSoftwareVersions();
-				bool bIsSdkValid = TurnkeyUtils.IsValueValid(InstalledSdk, AllowedSdks, Platform);
+				bool bIsAutoSdkValid = SDK.IsVersionValid(AutoSDKVersion, bForAutoSDK:true);
+				bool bIsManualSdkValid = SDK.IsVersionValid(ManualSDKVersion, bForAutoSDK:false);
 				TurnkeyUtils.Log("  Platform: {0}", TargetPlatform.ToString());
-				TurnkeyUtils.Log("  Installed Sdk: {0}", InstalledSdk);
-				TurnkeyUtils.Log("  Allowed Sdk(s): {0}", AllowedSdks);
-				TurnkeyUtils.Log("  Valid SDK Installed? {0}", bIsSdkValid);
+				TurnkeyUtils.Log("  Installed Manual Sdk: {0}", ManualSDKVersion);
+				TurnkeyUtils.Log("  Installed Auto Sdk: {0}", AutoSDKVersion);
+				TurnkeyUtils.Log("  Allowed Sdk Range: {0}-{1}", MinAllowedVersion, MaxAllowedVersion);
+				TurnkeyUtils.Log("  Valid Manual SDK Installed? {0}", bIsManualSdkValid);
+				TurnkeyUtils.Log("  Valid Auto SDK Installed? {0}", bIsAutoSdkValid);
 
 				// look for available sdks
 				List<SdkInfo> MatchingFullSdks = TurnkeyManifest.GetDiscoveredSdks()?.FindAll(x => x.Type == SdkInfo.SdkType.Full && x.IsValid(TargetPlatform));
diff --git a/Engine/Source/Programs/AutomationTool/Turnkey/Commands/VerifySdk.cs b/Engine/Source/Programs/AutomationTool/Turnkey/Commands/VerifySdk.cs
--- a/Engine/Source/Programs/AutomationTool/Turnkey/Commands/VerifySdk.cs
+++ b/Engine/Source/Programs/AutomationTool/Turnkey/Commands/VerifySdk.cs
@@ -6,6 +6,7 @@ using System;
 using System.Collections.Generic;
 using System.IO;
 using System.Linq;
+using Tools.DotNETCommon;
 using UnrealBuildTool;
 
 namespace Turnkey.Commands
@@ -110,7 +111,9 @@ namespace Turnkey.Commands
 
 				if ((LocalState & (SdkInfo.LocalAvailability.AutoSdk_ValidVersionExists | SdkInfo.LocalAvailability.InstalledSdk_ValidVersionExists)) == 0)
 				{
-					TurnkeyUtils.Report("{0}: Invalid: [No AutoSdk or Installed Sdk found matching {1} - {2}]", Platform, AutomationPlatform.GetAllowedSdks(), LocalState.ToString());
+					string MinAllowedVersion, MaxAllowedVersion;
+					UEBuildPlatformSDK.GetSDKForPlatform(Platform.ToString()).GetValidVersionRange(out MinAllowedVersion, out MaxAllowedVersion);
+					TurnkeyUtils.Report("{0}: Invalid: [No AutoSdk or Installed Sdk in range {1}-{2} - {3}]", Platform, MinAllowedVersion, MaxAllowedVersion, LocalState.ToString());
 					TurnkeyUtils.ExitCode = AutomationTool.ExitCode.Error_SDKNotFound;
 				}
 				else
@@ -119,17 +122,11 @@ namespace Turnkey.Commands
 					//					TurnkeyUtils.Log("{0}: Valid [Installed: '{1}', Required: '{2}']", Platform, PlatformObject.GetInstalledSdk(), PlatformObject.GetAllowedSdks());
 				}
 
-				// track if we installed an AutoSdk since we may need to setup autosdks before checking for devices
-				bool bAutoSdkSetupSucceeded = false;
-
 				// install if out of date, or if forcing it
 				if (bForceSdkInstall || (bUpdateIfNeeded && TurnkeyUtils.ExitCode != AutomationTool.ExitCode.Success))
 				{
 					TurnkeyUtils.ExitCode = AutomationTool.ExitCode.Success;
 
-					// if autosdk worked, we don't need to install full sdk
-					if (!bAutoSdkSetupSucceeded)
-					{
 // 						if (!bUnattended)
 // 						{
 // 							string Response = TurnkeyUtils.ReadInput("Your Sdk installation is not up to date. Would you like to install a valid Sdk? [Y/n]", "Y");
@@ -139,46 +136,38 @@ namespace Turnkey.Commands
 // 							}
 // 						}
 
-						SdkInfo BestSdk = null;
-						// find the best Sdk, prioritizing as request
-						if (bPreferFullSdk)
-						{
-							BestSdk = SdkInfo.FindMatchingSdk(AutomationPlatform, new SdkInfo.SdkType[] { SdkInfo.SdkType.Full, SdkInfo.SdkType.BuildOnly, SdkInfo.SdkType.AutoSdk }, bSelectBest: bUnattended);
-						}
-						else
-						{
-							BestSdk = SdkInfo.FindMatchingSdk(AutomationPlatform, new SdkInfo.SdkType[] { SdkInfo.SdkType.AutoSdk, SdkInfo.SdkType.BuildOnly, SdkInfo.SdkType.Full }, bSelectBest: bUnattended);
-						}
+					SdkInfo BestSdk = null;
+					// find the best Sdk, prioritizing as request
+					if (bPreferFullSdk)
+					{
+						BestSdk = SdkInfo.FindMatchingSdk(AutomationPlatform, new SdkInfo.SdkType[] { SdkInfo.SdkType.Full, SdkInfo.SdkType.BuildOnly, SdkInfo.SdkType.AutoSdk }, bSelectBest: bUnattended);
+					}
+					else
+					{
+						BestSdk = SdkInfo.FindMatchingSdk(AutomationPlatform, new SdkInfo.SdkType[] { SdkInfo.SdkType.AutoSdk, SdkInfo.SdkType.BuildOnly, SdkInfo.SdkType.Full }, bSelectBest: bUnattended);
+					}
 
-						if (BestSdk == null)
-						{
-							TurnkeyUtils.Log("ERROR: {0}: Unable to find any Sdks that could be installed");
-							TurnkeyUtils.ExitCode = AutomationTool.ExitCode.Error_SDKNotFound;
-							continue;
-						}
+					if (BestSdk == null)
+					{
+						TurnkeyUtils.Log("ERROR: {0}: Unable to find any Sdks that could be installed", Platform);
+						TurnkeyUtils.ExitCode = AutomationTool.ExitCode.Error_SDKNotFound;
+						continue;
+					}
 
-						TurnkeyUtils.Log("Will install {0}", BestSdk.DisplayName);
+					TurnkeyUtils.Log("Will install {0}", BestSdk.DisplayName);
 
-						bool bWasSdkInstalled = BestSdk.Install(Platform, null, bUnattended);
+					bool bWasSdkInstalled = BestSdk.Install(Platform, null, bUnattended);
 
-						// update LocalState
-						LocalState = SdkInfo.GetLocalAvailability(AutomationPlatform);
+					// update LocalState
+					LocalState = SdkInfo.GetLocalAvailability(AutomationPlatform);
 
-						// @todo turnkey: validate!
-					}
+					// @todo turnkey: validate!
 				}
 
 				// use the per-platform device list, unless it's not specifed, then use the global device list (which is set when not using platform specifiers)
 				List<string> DeviceNames = Pair.Value != null ? Pair.Value : SplitDeviceList;
 				if (DeviceNames != null && DeviceNames.Count > 0)
 				{
-					// if we just setup AutoSdk, then run the setup so that we have a hope of finding devices
-					if (bAutoSdkSetupSucceeded)
-					{
-						// @todo turnkey: Run the autosdk thing
-						//							AutomationPlatform.
-					}
-
 					DeviceInfo[] Devices = AutomationPlatform.GetDevices();
 					if (Devices == null)
 					{
diff --git a/Engine/Source/Programs/AutomationTool/Turnkey/SdkInfo.cs b/Engine/Source/Programs/AutomationTool/Turnkey/SdkInfo.cs
--- a/Engine/Source/Programs/AutomationTool/Turnkey/SdkInfo.cs
+++ b/Engine/Source/Programs/AutomationTool/Turnkey/SdkInfo.cs
@@ -9,7 +9,6 @@ using Tools.DotNETCommon;
 using UnrealBuildTool;
 using AutomationTool;
 using System.Linq;
-using System.Windows.Forms;
 
 namespace Turnkey
 {
@@ -89,6 +88,18 @@ namespace Turnkey
 				Clone.CustomSdkInputFiles = NewInstallers.ToArray();
 			}
 
+			// allow for shared SDK "platform" types
+			List<string> Platforms = TurnkeyUtils.ExpandVariables(Clone.PlatformString).ToLower().Split(",".ToCharArray()).ToList();
+			foreach (UEBuildPlatformSDK SDK in UEBuildPlatformSDK.AllSDKs)
+			{
+				// if the platforms contains a AutoSDK platform that doesn't match a real platform name, add the platform to the set
+				if (Platforms.Contains(SDK.GetAutoSDKPlatformName().ToLower()) && !Platforms.Contains(SDK.PlatformName.ToLower()))
+				{
+					Clone.PlatformString += "," + SDK.PlatformName;
+					Platforms.Add(SDK.PlatformName.ToLower()); ;
+				}
+			}
+
 			Clone.PostDeserialize();
 
 			return Clone;
@@ -119,6 +130,8 @@ namespace Turnkey
 		{
 			LocalAvailability Result = LocalAvailability.None;
 
+			UEBuildPlatformSDK SDK = UEBuildPlatformSDK.GetSDKForPlatform(AutomationPlatform.PlatformType.ToString());
+
 			// if we have the variable at all, 
 			string AutoSdkVar = Environment.GetEnvironmentVariable("UE_SDKS_ROOT");
 			if (AutoSdkVar != null)
@@ -126,15 +139,15 @@ namespace Turnkey
 				Result |= LocalAvailability.AutoSdk_VariableExists;
 
 				// get platform subdirectory
-				string AutoSubdir = string.Format("Host{0}/{1}", HostPlatform.Current.HostEditorPlatform.ToString(), AutomationPlatform.GetAutoSdkPlatformName());
+				string AutoSubdir = string.Format("Host{0}/{1}", HostPlatform.Current.HostEditorPlatform.ToString(), SDK.GetAutoSDKPlatformName());
 				DirectoryInfo PlatformDir = new DirectoryInfo(Path.Combine(AutoSdkVar, AutoSubdir));
 				if (PlatformDir.Exists)
 				{
 					bool bValidVersionFound = false;
 					foreach (DirectoryInfo Dir in PlatformDir.EnumerateDirectories())
 					{
-						// check for valid version number
-						if (TurnkeyUtils.IsValueValid(Dir.Name, AutomationPlatform.GetAllowedSdks(), AutomationPlatform))
+						// check for valid version number (auto SDK must match version exactly)
+						if (string.Compare(Dir.Name, SDK.GetDesiredVersion()) == 0)
 						{
 							// make sure it actually has bits in it
 							if (File.Exists(Path.Combine(Dir.FullName, "setup.bat")))
@@ -151,10 +164,10 @@ namespace Turnkey
 			}
 
 			// if anything is installed, this will return a value
-			string InstalledVersion = AutomationPlatform.GetInstalledSdk();
+			string InstalledVersion = SDK.GetInstalledVersion();
 			if (!string.IsNullOrEmpty(InstalledVersion))
 			{
-				if (TurnkeyUtils.IsValueValid(InstalledVersion, AutomationPlatform.GetAllowedSdks(), AutomationPlatform))
+				if (SDK.IsVersionValid(InstalledVersion, bForAutoSDK:false))
 				{
 					Result |= LocalAvailability.InstalledSdk_ValidVersionExists;
 				}
@@ -170,12 +183,14 @@ namespace Turnkey
 
 		static public SdkInfo FindMatchingSdk(AutomationTool.Platform Platform, SdkType[] TypePriority, bool bSelectBest, string DeviceType=null)
 		{
+			UEBuildPlatformSDK SDK = UEBuildPlatformSDK.GetSDKForPlatform(Platform.PlatformType.ToString());
+
 			foreach (SdkType Type in TypePriority)
 			{
 				List<SdkInfo> Sdks = TurnkeyManifest.GetDiscoveredSdks();
 				Sdks = Sdks.FindAll(x => x.SupportsPlatform(Platform.IniPlatformType));
 				Sdks = Sdks.FindAll(x => x.Type == Type);
-				Sdks = Sdks.FindAll(x => (x.Version == null && x.CustomSdkId != null) || TurnkeyUtils.IsValueValid(x.Version, Type == SdkType.Flash ? Platform.GetAllowedSoftwareVersions() : Platform.GetAllowedSdks(), Platform));
+				Sdks = Sdks.FindAll(x => (x.Version == null && x.CustomSdkId != null) || (Type == SdkType.Flash ? TurnkeyUtils.IsValueValid(Platform.GetAllowedSoftwareVersions(), x.Version, Platform) : SDK.IsVersionValid(x.Version, bForAutoSDK:x.Type==SdkType.AutoSdk)));
 				if (DeviceType != null)
 				{
 					Sdks = Sdks.FindAll(x => TurnkeyUtils.IsValueValid(DeviceType, x.AllowedFlashDeviceTypes, Platform));
@@ -187,7 +202,7 @@ namespace Turnkey
 					continue;
 				}
 
-				// if one was founf return it!
+				// if one was found return it!
 				if (Sdks.Count == 1)
 				{
 					return Sdks[0];
@@ -210,7 +225,7 @@ namespace Turnkey
 					return Best;
 				}
 
-				// if unable to pock one automatically, ask the user
+				// if unable to pick one automatically, ask the user
 				int Choice = TurnkeyUtils.ReadInputInt("Multiple Sdks found that could be installed. Please select one:", Sdks.Select(x => x.DisplayName).ToList(), true);
 
 				// we take canceling to mean to try the next type
@@ -294,7 +309,7 @@ namespace Turnkey
 			}
 			else
 			{
-				return TurnkeyUtils.IsValueValid(Version, AutomationPlatforms[Platform].GetAllowedSdks(), AutomationPlatforms[Platform]);
+				return UEBuildPlatformSDK.GetSDKForPlatform(Platform.ToString()).IsVersionValid(Version, bForAutoSDK:(Type == SdkType.AutoSdk));
 			}
 		}
 
@@ -350,7 +365,8 @@ namespace Turnkey
 			}
 			else
 			{
-				return !TurnkeyUtils.IsValueValid(AutomationPlatforms[Platform].GetInstalledSdk(), AutomationPlatforms[Platform].GetAllowedSdks(), AutomationPlatforms[Platform]);
+				UEBuildPlatformSDK SDK = UEBuildPlatformSDK.GetSDKForPlatform(Platform.ToString());
+				return !SDK.IsVersionValid(SDK.GetInstalledVersion(), bForAutoSDK:(Type == SdkType.AutoSdk));
 			}
 		}
 
@@ -359,7 +375,7 @@ namespace Turnkey
 			if (Type == SdkType.AutoSdk && !bSkipAutoSdkSetup)
 			{
 				// AutoSdk has some extra setup needed
-				if (SdkInfo.ConditionalSetupAutoSdk(Platform, bUnattended))
+				if (SdkInfo.ConditionalSetupAutoSdk(this, Platform, bUnattended))
 				{
 					return true;
 				}
@@ -440,7 +456,7 @@ namespace Turnkey
 			return bSucceeded;
 		}
 
-		public static bool ConditionalSetupAutoSdk(UnrealTargetPlatform Platform, bool bUnattended)
+		public static bool ConditionalSetupAutoSdk(SdkInfo Sdk, UnrealTargetPlatform Platform, bool bUnattended)
 		{
 			bool bAttemptAutoSdkSetup = false;
 			bool bSetupEnvVarAfterInstall = false;
@@ -468,75 +484,68 @@ namespace Turnkey
 
 				TurnkeyUtils.Log("{0}: AutoSdk is setup on this computer, will look for available AutoSdk to download", Platform);
 
-				SdkInfo MatchingAutoSdk = SdkInfo.FindMatchingSdk(AutomationPlatform, new SdkType[] { SdkInfo.SdkType.AutoSdk } , bSelectBest: true);
+				// make sure this is unset so that we can know if it worked or not after install
+				TurnkeyUtils.ClearVariable("CopyOutputPath");
+
+				// now download it (AutoSdks don't "install") on download
+				// @todo turnkey: handle errors, handle p4 going to wrong location, handle one Sdk for multiple platforms
+				Sdk.Install(Platform, null, bUnattended, bSkipAutoSdkSetup:true);
 
-				if (MatchingAutoSdk == null)
+				if (bSetupEnvVarAfterInstall)
 				{
-					// no matching AutoSdk found - will fall through to look for full install
-					if (bSetupEnvVarAfterInstall)
+
+					// this is where we synced the Sdk to
+					string InstalledRoot = TurnkeyUtils.GetVariableValue("CopyOutputPath");
+
+					// failed to install, nothing we can do
+					if (string.IsNullOrEmpty(InstalledRoot))
 					{
-						TurnkeyUtils.Log("{0}: Unable to find a matching AutoSdk, skipping AutoSdk setup", Platform);
+						TurnkeyUtils.ExitCode = AutomationTool.ExitCode.Error_SDKNotFound;
+						return false;
 					}
-				}
-				else
-				{
-					// make sure this is unset so that we can know if it worked or not after install
-					TurnkeyUtils.ClearVariable("CopyOutputPath");
 
-					// now download it (AutoSdks don't "install") on download
-					// @todo turnkey: handle errors, handle p4 going to wrong location, handle one Sdk for multiple platforms
-					MatchingAutoSdk.Install(Platform, null, bUnattended, bSkipAutoSdkSetup:true);
-
-					if (bSetupEnvVarAfterInstall)
+					// walk up to one above Host* directory
+					DirectoryInfo AutoSdkSearch;
+					if (Directory.Exists(InstalledRoot))
+					{
+						AutoSdkSearch = new DirectoryInfo(InstalledRoot);
+					}
+					else
+					{
+						AutoSdkSearch = new FileInfo(InstalledRoot).Directory;
+					}
+					while (AutoSdkSearch.Name != "Host" + HostPlatform.Current.HostEditorPlatform.ToString())
 					{
+						AutoSdkSearch = AutoSdkSearch.Parent;
+					}
 
-						// this is where we synced the Sdk to
-						string InstalledRoot = TurnkeyUtils.GetVariableValue("CopyOutputPath");
+					// now go one up to the parent of Host
+					AutoSdkSearch = AutoSdkSearch.Parent;
 
-						// failed to install, nothing we can do
-						if (string.IsNullOrEmpty(InstalledRoot))
+					string AutoSdkDir = AutoSdkSearch.FullName;
+					if (!bUnattended)
+					{
+						string Response = TurnkeyUtils.ReadInput("Enter directory for root of AutoSdks. Use detected value, or enter another:", AutoSdkSearch.FullName);
+						if (string.IsNullOrEmpty(Response))
 						{
-							TurnkeyUtils.ExitCode = AutomationTool.ExitCode.Error_SDKNotFound;
 							return false;
 						}
+					}
 
-						// walk up to one above Host* directory
-						DirectoryInfo AutoSdkSearch;
-						if (Directory.Exists(InstalledRoot))
-						{
-							AutoSdkSearch = new DirectoryInfo(InstalledRoot);
-						}
-						else
-						{
-							AutoSdkSearch = new FileInfo(InstalledRoot).Directory;
-						}
-						while (AutoSdkSearch.Name != "Host" + HostPlatform.Current.HostEditorPlatform.ToString())
-						{
-							AutoSdkSearch = AutoSdkSearch.Parent;
-						}
+					// set the env var, globally
+					TurnkeyUtils.StartTrackingExternalEnvVarChanges();
+					Environment.SetEnvironmentVariable("UE_SDKS_ROOT", AutoSdkDir);
+					Environment.SetEnvironmentVariable("UE_SDKS_ROOT", AutoSdkDir, EnvironmentVariableTarget.User);
+					TurnkeyUtils.EndTrackingExternalEnvVarChanges();
 
-						// now go one up to the parent of Host
-						AutoSdkSearch = AutoSdkSearch.Parent;
+				}
 
-						string AutoSdkDir = AutoSdkSearch.FullName;
-						if (!bUnattended)
-						{
-							string Response = TurnkeyUtils.ReadInput("Enter directory for root of AutoSdks. Use detected value, or enter another:", AutoSdkSearch.FullName);
-							if (string.IsNullOrEmpty(Response))
-							{
-								return false;
-							}
-						}
+				// and now activate it in case we need it this run
+				TurnkeyUtils.Log("Re-activating AutoSDK '{0}'...", Sdk.DisplayName);
 
-						// set the env var, globally
-						TurnkeyUtils.StartTrackingExternalEnvVarChanges();
-						Environment.SetEnvironmentVariable("UE_SDKS_ROOT", AutoSdkDir);
-						Environment.SetEnvironmentVariable("UE_SDKS_ROOT", AutoSdkDir, EnvironmentVariableTarget.User);
-						TurnkeyUtils.EndTrackingExternalEnvVarChanges();
-					}
+				Tools.DotNETCommon.UEBuildPlatformSDK.GetSDKForPlatform(Platform.ToString()).ReactivateAutoSDK();
 
-					return true;
-				}
+				return true;
 			}
 
 			return false;
diff --git a/Engine/Source/Programs/AutomationTool/Turnkey/Turnkey.Automation.cs b/Engine/Source/Programs/AutomationTool/Turnkey/Turnkey.Automation.cs
--- a/Engine/Source/Programs/AutomationTool/Turnkey/Turnkey.Automation.cs
+++ b/Engine/Source/Programs/AutomationTool/Turnkey/Turnkey.Automation.cs
@@ -5,6 +5,7 @@ using System.Collections.Generic;
 using Tools.DotNETCommon;
 using UnrealBuildTool;
 using AutomationTool;
+using System.Threading;
 
 namespace Turnkey
 {
@@ -12,14 +13,33 @@ namespace Turnkey
 	{
 		public override ExitCode Execute()
 		{
+			// 			Thread t = new Thread(Turnkey.ThreadProc);
+			// 			Console.WriteLine("Before setting apartment state: {0}",
+			// 				t.GetApartmentState());
+			// 
+			// 			t.SetApartmentState(ApartmentState.STA);
+			// 			Console.WriteLine("After setting apartment state: {0}",
+			// 				t.GetApartmentState());
+			// 
+			// 			t.Start(this);
+			// 			t.Join();
+
+			ThreadProc(this);
+			return TurnkeyUtils.ExitCode;
+		}
+
+		private static void ThreadProc(object Data)
+		{
+			BuildCommand Build = (BuildCommand)Data;
+
 			IOProvider IOProvider;
-			if (ParseParam("EditorIO"))
+			if (Build.ParseParam("EditorIO"))
 			{
 				IOProvider = new HybridIOProvider();
 			}
 			else
 			{
-				string ReportFilename = ParseParamValue("ReportFilename");
+				string ReportFilename = Build.ParseParamValue("ReportFilename");
 				if (!string.IsNullOrEmpty(ReportFilename))
 				{
 					IOProvider = new ReportIOProvider(ReportFilename);
@@ -30,9 +50,10 @@ namespace Turnkey
 				}
 			}
 
-			return Turnkey.Execute(IOProvider, this);
+			Turnkey.Execute(IOProvider, Build);
 		}
 
+
 		public static AutomationTool.ExitCode Execute(IOProvider IOProvider, BuildCommand CommandUtilHelper)
 		{
 			TurnkeyUtils.Initialize(IOProvider, CommandUtilHelper);
diff --git a/Engine/Source/Programs/AutomationTool/Turnkey/Turnkey.Automation.csproj b/Engine/Source/Programs/AutomationTool/Turnkey/Turnkey.Automation.csproj
--- a/Engine/Source/Programs/AutomationTool/Turnkey/Turnkey.Automation.csproj
+++ b/Engine/Source/Programs/AutomationTool/Turnkey/Turnkey.Automation.csproj
@@ -65,6 +65,11 @@
     <Compile Include="TurnkeyUtils.cs" />
   </ItemGroup>
   <ItemGroup>
+    <ProjectReference Include="..\..\DotNETCommon\BuildUtilities\BuildUtilities.csproj">
+      <Project>{5965bed9-4823-4671-84a5-c8d6503d5848}</Project>
+      <Name>BuildUtilities</Name>
+      <Private>False</Private>
+    </ProjectReference>
     <ProjectReference Include="..\..\DotNETCommon\DotNETUtilities\DotNETUtilities.csproj">
       <Project>{5d7d66e8-8c76-4af9-b3ec-2ef03421d730}</Project>
       <Name>DotNETUtilities</Name>
@@ -104,6 +109,7 @@
     <Reference Include="System.Windows.Forms" />
     <Reference Include="System.Xml" />
     <Reference Include="System.Xml.Serialization" />
+    <Reference Include="Microsoft.VisualBasic" />
   </ItemGroup>
   <ItemGroup />
   <Import Project="$(MSBuildToolsPath)\Microsoft.CSharp.targets" />
diff --git a/Engine/Source/Programs/AutomationTool/Turnkey/TurnkeyUtils.cs b/Engine/Source/Programs/AutomationTool/Turnkey/TurnkeyUtils.cs
--- a/Engine/Source/Programs/AutomationTool/Turnkey/TurnkeyUtils.cs
+++ b/Engine/Source/Programs/AutomationTool/Turnkey/TurnkeyUtils.cs
@@ -297,7 +297,7 @@ namespace Turnkey
 
 				// convert inputs to uint (unless already converted above by platform)
 				UInt64 ValueInt = 0;
-				if (!TryConvertToUint64(Value, out ValueInt) && !Platform.TryConvertVersionToInt(Value, out ValueInt))
+				if (!TryConvertToUint64(Value, out ValueInt) && !UEBuildPlatformSDK.GetSDKForPlatform(Platform.ToString()).TryConvertVersionToInt(Value, out ValueInt))
 				{
 					TurnkeyUtils.Log("Warning: range: input value [{0}] was not an unsigned integer, and platform couldn't convert it", Value);
 					return false;
diff --git a/Engine/Source/Programs/DotNETCommon/BuildUtilities/BuildUtilities.csproj b/Engine/Source/Programs/DotNETCommon/BuildUtilities/BuildUtilities.csproj
new file mode 100644
--- /dev/null
+++ b/Engine/Source/Programs/DotNETCommon/BuildUtilities/BuildUtilities.csproj
@@ -0,0 +1,74 @@
+﻿<?xml version="1.0" encoding="utf-8"?>
+<Project ToolsVersion="12.0" DefaultTargets="Build" xmlns="http://schemas.microsoft.com/developer/msbuild/2003">
+  <Import Project="$(MSBuildExtensionsPath)\$(MSBuildToolsVersion)\Microsoft.Common.props" Condition="Exists('$(MSBuildExtensionsPath)\$(MSBuildToolsVersion)\Microsoft.Common.props')" />
+  <PropertyGroup>
+    <Configuration Condition=" '$(Configuration)' == '' ">Development</Configuration>
+    <Platform Condition=" '$(Platform)' == '' ">AnyCPU</Platform>
+    <ProjectGuid>{5965BED9-4823-4671-84A5-C8D6503D5848}</ProjectGuid>
+    <OutputType>Library</OutputType>
+    <AppDesignerFolder>Properties</AppDesignerFolder>
+    <RootNamespace>Tools.DotNETCommon</RootNamespace>
+    <AssemblyName>BuildUtilities</AssemblyName>
+    <TargetFrameworkVersion>v4.6.2</TargetFrameworkVersion>
+    <FileAlignment>512</FileAlignment>
+    <TargetFrameworkProfile />
+  </PropertyGroup>
+  <PropertyGroup Condition=" '$(Configuration)|$(Platform)' == 'Debug|AnyCPU' ">
+    <DebugSymbols>true</DebugSymbols>
+    <DebugType>full</DebugType>
+    <Optimize>false</Optimize>
+    <OutputPath>..\..\..\..\Binaries\DotNET\</OutputPath>
+    <DefineConstants>DEBUG;TRACE</DefineConstants>
+    <ErrorReport>prompt</ErrorReport>
+    <WarningLevel>4</WarningLevel>
+    <DocumentationFile>
+    </DocumentationFile>
+  </PropertyGroup>
+  <PropertyGroup Condition="'$(Configuration)|$(Platform)' == 'Development|AnyCPU'">
+    <OutputPath>..\..\..\..\Binaries\DotNET\</OutputPath>
+    <DefineConstants>TRACE</DefineConstants>
+    <DocumentationFile>
+    </DocumentationFile>
+    <Optimize>true</Optimize>
+    <DebugType>pdbonly</DebugType>
+    <PlatformTarget>AnyCPU</PlatformTarget>
+    <ErrorReport>prompt</ErrorReport>
+  </PropertyGroup>
+  <PropertyGroup Condition="'$(Configuration)|$(Platform)' == 'Release|AnyCPU'">
+    <OutputPath>..\..\..\..\Binaries\DotNET\</OutputPath>
+    <DefineConstants>TRACE</DefineConstants>
+    <Optimize>true</Optimize>
+    <DebugType>pdbonly</DebugType>
+    <PlatformTarget>AnyCPU</PlatformTarget>
+    <ErrorReport>prompt</ErrorReport>
+  </PropertyGroup>
+  <ItemGroup>
+    <Reference Include="System" />
+    <Reference Include="System.Core" />
+    <Reference Include="System.Xml.Linq" />
+    <Reference Include="System.Data.DataSetExtensions" />
+    <Reference Include="Microsoft.CSharp" />
+    <Reference Include="System.Data" />
+    <Reference Include="System.Xml" />
+  </ItemGroup>
+  <ItemGroup>
+    <Compile Include="..\MetaData.cs">
+      <Link>Properties\MetaData.cs</Link>
+    </Compile>
+    <Compile Include="UEBuildPlatformSDK.cs" />
+  </ItemGroup>
+  <ItemGroup>
+    <ProjectReference Include="..\DotNETUtilities\DotNETUtilities.csproj">
+      <Project>{5d7d66e8-8c76-4af9-b3ec-2ef03421d730}</Project>
+      <Name>DotNETUtilities</Name>
+    </ProjectReference>
+  </ItemGroup>
+  <Import Project="$(MSBuildToolsPath)\Microsoft.CSharp.targets" />
+  <!-- To modify your build process, add your task inside one of the targets below and uncomment it. 
+       Other similar extension points exist, see Microsoft.Common.targets.
+  <Target Name="BeforeBuild">
+  </Target>
+  <Target Name="AfterBuild">
+  </Target>
+  -->
+</Project>
\ No newline at end of file
diff --git a/Engine/Source/Programs/DotNETCommon/BuildUtilities/UEBuildPlatformSDK.cs b/Engine/Source/Programs/DotNETCommon/BuildUtilities/UEBuildPlatformSDK.cs
new file mode 100644
--- /dev/null
+++ b/Engine/Source/Programs/DotNETCommon/BuildUtilities/UEBuildPlatformSDK.cs
@@ -0,0 +1,1135 @@
+// Copyright Epic Games, Inc. All Rights Reserved.
+
+using System;
+using System.Collections.Generic;
+using System.Diagnostics;
+using System.IO;
+using System.Linq;
+using System.Text;
+
+namespace Tools.DotNETCommon
+{
+	/// <summary>
+	/// SDK installation status
+	/// </summary>
+	public enum SDKStatus
+	{
+		/// <summary>
+		/// Desired SDK is installed and set up.
+		/// </summary>
+		Valid,
+
+		/// <summary>
+		/// Could not find the desired SDK, SDK setup failed, etc.		
+		/// </summary>
+		Invalid,
+	};
+
+	/// <summary>
+	/// SDK for a platform
+	/// </summary>
+	abstract public class UEBuildPlatformSDK
+	{
+		// Public SDK handling, not specific to AutoSDK
+
+		public UEBuildPlatformSDK()
+		{
+		}
+
+		#region Global SDK Registration
+
+		/// <summary>
+		/// Registers the SDK for a given platform (as a string, but equivalent to UnrealTargetPlatform)
+		/// </summary>
+		/// <param name="SDK">SDK object</param>
+		/// <param name="PlatformName">Platform name for this SDK</param>
+		public static void RegisterSDKForPlatform(UEBuildPlatformSDK SDK, string PlatformName)
+		{
+			// verify that neither platform or sdk were added before
+			if (SDKRegistry.Count(x => x.Key == PlatformName || x.Value == SDK) > 0)
+			{
+				throw new Exception(string.Format("Re-registering SDK for {0}. All Platforms must have a unique SDK object", PlatformName));
+			}
+
+			SDKRegistry.Add(PlatformName, SDK);
+
+			SDK.Init(PlatformName);
+		}
+
+		private void Init(string InPlatformName)
+		{
+			PlatformName = InPlatformName;
+
+			// if the parent set up autosdk, the env vars will be wrong, but we can still get the manual SDK version from before it was setup
+			if (HasParentProcessSetupAutoSDK(out CachedManualSDKVersion))
+			{
+				// we pass along __None to indicate the parent didn't have a manual sdk installed
+				if (CachedManualSDKVersion == "__None")
+				{
+					CachedManualSDKVersion = null;
+				}
+			}
+			else
+			{
+				// if there was no parent, get the SDK version before we run AutoSDK to get the manual version
+				CachedManualSDKVersion = GetInstalledSDKVersion();
+			}
+		}
+
+		#endregion
+
+		#region Main Public Interface/Utilties
+
+		/// <summary>
+		/// Retrieves a previously registered SDK for a given platform
+		/// </summary>
+		/// <param name="PlatformName">String name of the platform (equivalent to UnrealTargetPlatform)</param>
+		/// <returns></returns>
+		public static UEBuildPlatformSDK GetSDKForPlatform(string PlatformName)
+		{
+			UEBuildPlatformSDK SDK;
+			SDKRegistry.TryGetValue(PlatformName, out SDK);
+
+			return SDK;
+		}
+
+		/// <summary>
+		/// Gets the set of all known SDKs
+		/// </summary>
+		public static UEBuildPlatformSDK[] AllSDKs
+		{
+			get	{ return SDKRegistry.Values.ToArray(); }
+		}
+
+		// String name of the platform (will match an UnrealTargetPlatform)
+		public string PlatformName;
+
+		public string GetInstalledVersion(out bool bIsAutoSDK)
+		{
+			bIsAutoSDK = HasSetupAutoSDK();
+			return GetInstalledSDKVersion();
+		}
+
+		public string GetInstalledVersion()
+		{
+			return GetInstalledSDKVersion();
+		}
+
+		public void GetInstalledVersions(out string ManualSDKVersion, out string AutoSDKVersion)
+		{
+			// if we support AutoSDKs, then return both versions
+			if (PlatformSupportsAutoSDKs())
+			{
+				// we don't need ask the platform for the installed version because we know the AutoSDK version has to match the DesiredVersion if it was setup properly
+				//AutoSDKVersion = (HasRequiredAutoSDKInstalled() == SDKStatus.Valid) ? GetDesiredVersion() : null;
+				AutoSDKVersion = GetInstalledSDKVersion();
+			}
+			else
+			{
+				AutoSDKVersion = null;
+				if (CachedManualSDKVersion != GetInstalledSDKVersion())
+				{
+					throw new Exception("Manual SDK version changed, this is not supported yet");
+				}
+			}
+
+			ManualSDKVersion = CachedManualSDKVersion;
+		}
+
+
+		public virtual bool IsVersionValid(string Version, bool bForAutoSDK)
+		{
+			return IsVersionValidInternal(Version, bForAutoSDK);
+		}
+
+		public void ReactivateAutoSDK()
+		{
+			// @todo turnkey: this needs to force re-doing it, as it is likely a no-op, need to investigate what to clear out
+			ManageAndValidateSDK();
+		}
+
+		#endregion
+
+		#region Platform Overrides
+
+		/// <summary>
+		/// Returns the installed SDK version, used to determine if up to date or not (
+		/// </summary>
+		/// <returns></returns>
+		public abstract string GetInstalledSDKVersion();
+
+		/// <summary>
+		/// Return the SDK version that the platform wants to use (AutoSDK dir must match this, full SDKs can be in a valid range)
+		/// </summary>
+		/// <returns></returns>
+		public virtual string GetDesiredVersion()
+		{
+			return GetRequiredSDKString();
+		}
+
+		/// <summary>
+		/// Gets the valid (integer) range of Sdk versions. Must be an integer to easily check a range vs a particular version
+		/// </summary>
+		/// <param name="MinVersion">Smallest version allowed</param>
+		/// <param name="MaxVersion">Largest version allowed (inclusive)</param>
+		/// <returns>True if the versions are valid, false if the platform is unable to convert its versions into an integer</returns>
+		public virtual bool GetValidVersionRange(out UInt64 MinVersion, out UInt64 MaxVersion)
+		{
+			string MinVersionString, MaxVersionString;
+			GetValidVersionRange(out MinVersionString, out MaxVersionString);
+
+			if (TryConvertVersionToInt(MinVersionString, out MinVersion) && TryConvertVersionToInt(MaxVersionString, out MaxVersion))
+			{
+				return true;
+			}
+
+			// if a version can't be converted to an int, we can't get any versions
+			MinVersion = MaxVersion = 0;
+			return false;
+		}
+
+		/// <summary>
+		/// Gets the valid string range of Sdk versions. TryConvertVersionToInt() will need to succeed to make this usable for range checks
+		/// </summary>
+		/// <param name="MinVersion">Smallest version allowed</param>
+		/// <param name="MaxVersion">Largest version allowed (inclusive)</param>
+		/// <returns>True if the versions are valid, false if the platform is unable to convert its versions into an integer</returns>
+		public virtual void GetValidVersionRange(out string MinVersion, out string MaxVersion)
+		{
+			// @todo turnkey make this abstract?
+			// by default, use GetDesiredVersion() to min and max
+			MinVersion = MaxVersion = GetDesiredVersion();
+		}
+
+		// Let platform override behavior to determine if a version is a valid (useful for non-numeric versions)
+		public virtual bool IsVersionValidInternal(string Version, bool bForAutoSDK)
+		{
+			// we could have null if no SDK is installed at all, etc, which is always a failure
+			if (Version == null)
+			{
+				return false;
+			}
+
+			// short circuit range check if the Version is the desired version already
+			if (string.Compare(Version, GetDesiredVersion(), true) == 0)
+			{
+				return true;
+			}
+			else
+			{
+				// AutoSDK must match the desired version exactly, since that is the only one we will use
+				if (bForAutoSDK)
+				{
+					return false;
+				}
+
+			}
+
+			// convert it to an integer
+			UInt64 IntVersion;
+			if (!TryConvertVersionToInt(Version, out IntVersion))
+			{
+				return false;
+			}
+
+			// get numeric range
+			UInt64 MinVersion, MaxVersion;
+			GetValidVersionRange(out MinVersion, out MaxVersion);
+			return IntVersion >= MinVersion && IntVersion <= MaxVersion;
+		}
+
+
+		/// <summary>
+		/// Only the platform can convert a version string into an integer that is usable for comparison
+		/// </summary>
+		/// <param name="StringValue">Version that comes from the installed SDK or a Turnkey manifest or the like</param>
+		/// <param name="OutValue">The integer version of StringValue, can be used to compare against a valid range</param>
+		/// <returns>If the StringValue was able to be be converted to an integer</returns>
+		public virtual bool TryConvertVersionToInt(string StringValue, out UInt64 OutValue)
+		{
+			// @todo turnkey make this abstract?
+			OutValue = 0;
+			return false;
+		}
+
+		/// <summary>
+		/// Allow the platform SDK to override the name it will use in AutoSDK, but default to the platform name
+		/// </summary>
+		/// <returns>The name of the directory to use inside the AutoSDK system</returns>
+		public virtual string GetAutoSDKPlatformName()
+		{
+			return PlatformName;
+		}
+
+		#endregion
+
+
+
+
+
+
+		#region Private/Protected general functionality
+
+		// this is the SDK version that was set before activating AutoSDK, since AutoSDK may remove ability to retrieve the Manual SDK version
+		protected string CachedManualSDKVersion;
+		private static Dictionary<string, UEBuildPlatformSDK> SDKRegistry = new Dictionary<string, UEBuildPlatformSDK>();
+
+		#endregion
+
+		// AutoSDKs handling portion
+
+		#region protected AutoSDKs Utility
+
+		/// <summary>
+		/// Name of the file that holds currently install SDK version string
+		/// </summary>
+		protected const string CurrentlyInstalledSDKStringManifest = "CurrentlyInstalled.txt";
+
+		/// <summary>
+		/// name of the file that holds the last succesfully run SDK setup script version
+		/// </summary>
+		protected const string LastRunScriptVersionManifest = "CurrentlyInstalled.Version.txt";
+
+		/// <summary>
+		/// Name of the file that holds environment variables of current SDK
+		/// </summary>
+		protected const string SDKEnvironmentVarsFile = "OutputEnvVars.txt";
+
+		protected const string SDKRootEnvVar = "UE_SDKS_ROOT";
+
+		protected const string AutoSetupEnvVar = "AutoSDKSetup";
+
+		protected static bool IsWindows()
+		{
+			return Environment.OSVersion.Platform == PlatformID.Win32NT || Environment.OSVersion.Platform == PlatformID.Win32S || Environment.OSVersion.Platform == PlatformID.Win32Windows;
+		}
+
+		protected static bool IsMac()
+		{
+			// mono tends to return Unix on Mac, so check for a Mac-only file
+			return Environment.OSVersion.Platform == PlatformID.MacOSX || (Environment.OSVersion.Platform == PlatformID.Unix && File.Exists ("/System/Library/CoreServices/SystemVersion.plist"));
+		}
+
+		private static string GetAutoSDKHostPlatform()
+		{
+			if (IsWindows())
+			{
+				return "Win64";
+			}
+			else if (Environment.OSVersion.Platform == PlatformID.MacOSX)
+			{
+				return "Mac";
+			}
+			else if (Environment.OSVersion.Platform == PlatformID.Unix)
+			{
+				return "Linux";
+			}
+			throw new Exception("Unknown host platform!");
+		}
+
+		/// <summary>
+		/// Whether platform supports switching SDKs during runtime
+		/// </summary>
+		/// <returns>true if supports</returns>
+		protected virtual bool PlatformSupportsAutoSDKs()
+		{
+			return false;
+		}
+
+		static private bool bCheckedAutoSDKRootEnvVar = false;
+		static private bool bAutoSDKSystemEnabled = false;
+		static private bool HasAutoSDKSystemEnabled()
+		{
+			if (!bCheckedAutoSDKRootEnvVar)
+			{
+				string SDKRoot = Environment.GetEnvironmentVariable(SDKRootEnvVar);
+				if (SDKRoot != null)
+				{
+					bAutoSDKSystemEnabled = true;
+				}
+				bCheckedAutoSDKRootEnvVar = true;
+			}
+			return bAutoSDKSystemEnabled;
+		}
+
+		// Whether AutoSDK setup is safe. AutoSDKs will damage manual installs on some platforms.
+		protected bool IsAutoSDKSafe()
+		{
+			return !IsAutoSDKDestructive() || !HasAnyManualInstall();
+		}
+
+		/// <summary>
+		/// Returns SDK string as required by the platform
+		/// </summary>
+		/// <returns>Valid SDK string</returns>
+		protected virtual string GetRequiredSDKString()
+		{
+			return "";
+		}
+
+		/// <summary>
+		/// Gets the version number of the SDK setup script itself.  The version in the base should ALWAYS be the master revision from the last refactor.
+		/// If you need to force a rebuild for a given platform, override this for the given platform.
+		/// </summary>
+		/// <returns>Setup script version</returns>
+		protected virtual String GetRequiredScriptVersionString()
+		{
+			return "3.0";
+		}
+
+		/// <summary>
+		/// Returns path to platform SDKs
+		/// </summary>
+		/// <returns>Valid SDK string</returns>
+		protected string GetPathToPlatformAutoSDKs()
+		{
+			string SDKPath = "";
+			string SDKRoot = Environment.GetEnvironmentVariable(SDKRootEnvVar);
+			if (SDKRoot != null)
+			{
+				if (SDKRoot != "")
+				{
+					SDKPath = Path.Combine(SDKRoot, "Host" + GetAutoSDKHostPlatform(), GetAutoSDKPlatformName());
+				}
+			}
+			return SDKPath;
+		}
+
+		/// <summary>
+		/// Returns path to platform SDKs
+		/// </summary>
+		/// <returns>Valid SDK string</returns>
+		public static bool TryGetHostPlatformAutoSDKDir(out DirectoryReference OutPlatformDir)
+		{
+			string SDKRoot = Environment.GetEnvironmentVariable(SDKRootEnvVar);
+			if (String.IsNullOrEmpty(SDKRoot))
+			{
+				OutPlatformDir = null;
+				return false;
+			}
+			else
+			{
+				OutPlatformDir = DirectoryReference.Combine(new DirectoryReference(SDKRoot), "Host" + GetAutoSDKHostPlatform());
+				return true;
+			}
+		}
+
+		/// <summary>
+		/// Because most ManualSDK determination depends on reading env vars, if this process is spawned by a process that ALREADY set up
+		/// AutoSDKs then all the SDK env vars will exist, and we will spuriously detect a Manual SDK. (children inherit the environment of the parent process).
+		/// Therefore we write out an env variable to set in the command file (OutputEnvVars.txt) such that child processes can determine if their manual SDK detection
+		/// is bogus.  Make it platform specific so that platforms can be in different states.
+		/// </summary>
+		protected string GetPlatformAutoSDKSetupEnvVar()
+		{
+			return GetAutoSDKPlatformName() + AutoSetupEnvVar;
+		}
+
+		/// <summary>
+		/// Gets currently installed version
+		/// </summary>
+		/// <param name="PlatformSDKRoot">absolute path to platform SDK root</param>
+		/// <param name="OutInstalledSDKVersionString">version string as currently installed</param>
+		/// <returns>true if was able to read it</returns>
+		protected bool GetCurrentlyInstalledSDKString(string PlatformSDKRoot, out string OutInstalledSDKVersionString)
+		{
+			if (Directory.Exists(PlatformSDKRoot))
+			{
+				string VersionFilename = Path.Combine(PlatformSDKRoot, CurrentlyInstalledSDKStringManifest);
+				if (File.Exists(VersionFilename))
+				{
+					using (StreamReader Reader = new StreamReader(VersionFilename))
+					{
+						string Version = Reader.ReadLine();
+						string Type = Reader.ReadLine();
+
+						// don't allow ManualSDK installs to count as an AutoSDK install version.
+						if (Type != null && Type == "AutoSDK")
+						{
+							if (Version != null)
+							{
+								OutInstalledSDKVersionString = Version;
+								return true;
+							}
+						}
+					}
+				}
+			}
+
+			OutInstalledSDKVersionString = "";
+			return false;
+		}
+
+		/// <summary>
+		/// Gets the version of the last successfully run setup script.
+		/// </summary>
+		/// <param name="PlatformSDKRoot">absolute path to platform SDK root</param>
+		/// <param name="OutLastRunScriptVersion">version string</param>
+		/// <returns>true if was able to read it</returns>
+		protected bool GetLastRunScriptVersionString(string PlatformSDKRoot, out string OutLastRunScriptVersion)
+		{
+			if (Directory.Exists(PlatformSDKRoot))
+			{
+				string VersionFilename = Path.Combine(PlatformSDKRoot, LastRunScriptVersionManifest);
+				if (File.Exists(VersionFilename))
+				{
+					using (StreamReader Reader = new StreamReader(VersionFilename))
+					{
+						string Version = Reader.ReadLine();
+						if (Version != null)
+						{
+							OutLastRunScriptVersion = Version;
+							return true;
+						}
+					}
+				}
+			}
+
+			OutLastRunScriptVersion = "";
+			return false;
+		}
+
+		/// <summary>
+		/// Sets currently installed version
+		/// </summary>
+		/// <param name="InstalledSDKVersionString">SDK version string to set</param>
+		/// <returns>true if was able to set it</returns>
+		protected bool SetCurrentlyInstalledAutoSDKString(String InstalledSDKVersionString)
+		{
+			String PlatformSDKRoot = GetPathToPlatformAutoSDKs();
+			if (Directory.Exists(PlatformSDKRoot))
+			{
+				string VersionFilename = Path.Combine(PlatformSDKRoot, CurrentlyInstalledSDKStringManifest);
+				if (File.Exists(VersionFilename))
+				{
+					File.Delete(VersionFilename);
+				}
+
+				using (StreamWriter Writer = File.CreateText(VersionFilename))
+				{
+					Writer.WriteLine(InstalledSDKVersionString);
+					Writer.WriteLine("AutoSDK");
+					return true;
+				}
+			}
+
+			return false;
+		}
+
+		protected void SetupManualSDK()
+		{
+			if (PlatformSupportsAutoSDKs() && HasAutoSDKSystemEnabled())
+			{
+				String InstalledSDKVersionString = GetRequiredSDKString();
+				String PlatformSDKRoot = GetPathToPlatformAutoSDKs();
+                if (!Directory.Exists(PlatformSDKRoot))
+                {
+                    Directory.CreateDirectory(PlatformSDKRoot);
+                }
+
+				{
+					string VersionFilename = Path.Combine(PlatformSDKRoot, CurrentlyInstalledSDKStringManifest);
+					if (File.Exists(VersionFilename))
+					{
+						File.Delete(VersionFilename);
+					}
+
+					string EnvVarFile = Path.Combine(PlatformSDKRoot, SDKEnvironmentVarsFile);
+					if (File.Exists(EnvVarFile))
+					{
+						File.Delete(EnvVarFile);
+					}
+
+					using (StreamWriter Writer = File.CreateText(VersionFilename))
+					{
+						Writer.WriteLine(InstalledSDKVersionString);
+						Writer.WriteLine("ManualSDK");
+					}
+				}
+			}
+		}
+
+		protected bool SetLastRunAutoSDKScriptVersion(string LastRunScriptVersion)
+		{
+			String PlatformSDKRoot = GetPathToPlatformAutoSDKs();
+			if (Directory.Exists(PlatformSDKRoot))
+			{
+				string VersionFilename = Path.Combine(PlatformSDKRoot, LastRunScriptVersionManifest);
+				if (File.Exists(VersionFilename))
+				{
+					File.Delete(VersionFilename);
+				}
+
+				using (StreamWriter Writer = File.CreateText(VersionFilename))
+				{
+					Writer.WriteLine(LastRunScriptVersion);
+					return true;
+				}
+			}
+			return false;
+		}
+
+		/// <summary>
+		/// Returns Hook names as needed by the platform
+		/// (e.g. can be overridden with custom executables or scripts)
+		/// </summary>
+		/// <param name="Hook">Hook type</param>
+		protected virtual string GetHookExecutableName(SDKHookType Hook)
+		{
+			if (IsWindows())
+			{
+				if (Hook == SDKHookType.Uninstall)
+				{
+					return "unsetup.bat";
+				}
+				else
+				{
+					return "setup.bat";
+				}
+			}
+			else
+			{
+				if (Hook == SDKHookType.Uninstall)
+				{
+					return "unsetup.sh";
+				}
+				else
+				{
+					return "setup.sh";
+				}
+			}
+		}
+
+		/// <summary>
+		/// Runs install/uninstall hooks for SDK
+		/// </summary>
+		/// <param name="PlatformSDKRoot">absolute path to platform SDK root</param>
+		/// <param name="SDKVersionString">version string to run for (can be empty!)</param>
+		/// <param name="Hook">which one of hooks to run</param>
+		/// <param name="bHookCanBeNonExistent">whether a non-existing hook means failure</param>
+		/// <returns>true if succeeded</returns>
+		protected virtual bool RunAutoSDKHooks(string PlatformSDKRoot, string SDKVersionString, SDKHookType Hook, bool bHookCanBeNonExistent = true)
+		{
+			if (!IsAutoSDKSafe())
+			{
+				Log.TraceLog(GetAutoSDKPlatformName() + " attempted to run SDK hook which could have damaged manual SDK install!");
+				return false;
+			}
+			if (SDKVersionString != "")
+			{
+				string SDKDirectory = Path.Combine(PlatformSDKRoot, SDKVersionString);
+				string HookExe = Path.Combine(SDKDirectory, GetHookExecutableName(Hook));
+
+				if (File.Exists(HookExe))
+				{
+					Log.TraceLog("Running {0} hook {1}", Hook, HookExe);
+
+					// run it
+					Process HookProcess = new Process();
+					HookProcess.StartInfo.WorkingDirectory = SDKDirectory;
+					HookProcess.StartInfo.FileName = HookExe;
+					HookProcess.StartInfo.Arguments = "";
+					HookProcess.StartInfo.WindowStyle = ProcessWindowStyle.Hidden;
+
+					// seems to break the build machines?
+					//HookProcess.StartInfo.UseShellExecute = false;
+					//HookProcess.StartInfo.RedirectStandardOutput = true;
+					//HookProcess.StartInfo.RedirectStandardError = true;					
+
+//					using (ScopedTimer HookTimer = new ScopedTimer("Time to run hook: ", LogEventType.Log))
+					{
+						//installers may require administrator access to succeed. so run as an admmin.
+						HookProcess.StartInfo.Verb = "runas";
+						HookProcess.Start();
+						HookProcess.WaitForExit();
+					}
+
+					//LogAutoSDK(HookProcess.StandardOutput.ReadToEnd());
+					//LogAutoSDK(HookProcess.StandardError.ReadToEnd());
+					if (HookProcess.ExitCode != 0)
+					{
+						Log.TraceLog("Hook exited uncleanly (returned {0}), considering it failed.", HookProcess.ExitCode);
+						return false;
+					}
+
+					return true;
+				}
+				else
+				{
+					Log.TraceLog("File {0} does not exist", HookExe);
+				}
+			}
+			else
+			{
+				Log.TraceLog("Version string is blank for {0}. Can't determine {1} hook.", PlatformSDKRoot, Hook.ToString());
+			}
+
+			return bHookCanBeNonExistent;
+		}
+
+		/// <summary>
+		/// Loads environment variables from SDK
+		/// If any commands are added or removed the handling needs to be duplicated in
+		/// TargetPlatformManagerModule.cpp
+		/// </summary>
+		/// <param name="PlatformSDKRoot">absolute path to platform SDK</param>
+		/// <returns>true if succeeded</returns>
+		protected bool SetupEnvironmentFromAutoSDK(string PlatformSDKRoot)
+		{
+			string EnvVarFile = Path.Combine(PlatformSDKRoot, SDKEnvironmentVarsFile);
+			if (File.Exists(EnvVarFile))
+			{
+				using (StreamReader Reader = new StreamReader(EnvVarFile))
+				{
+					List<string> PathAdds = new List<string>();
+					List<string> PathRemoves = new List<string>();
+
+					List<string> EnvVarNames = new List<string>();
+					List<string> EnvVarValues = new List<string>();
+
+					bool bNeedsToWriteAutoSetupEnvVar = true;
+					String PlatformSetupEnvVar = GetPlatformAutoSDKSetupEnvVar();
+					for (; ; )
+					{
+						string VariableString = Reader.ReadLine();
+						if (VariableString == null)
+						{
+							break;
+						}
+
+						string[] Parts = VariableString.Split('=');
+						if (Parts.Length != 2)
+						{
+							Log.TraceLog("Incorrect environment variable declaration:");
+							Log.TraceLog(VariableString);
+							return false;
+						}
+
+						if (String.Compare(Parts[0], "strippath", true) == 0)
+						{
+							PathRemoves.Add(Parts[1]);
+						}
+						else if (String.Compare(Parts[0], "addpath", true) == 0)
+						{
+							PathAdds.Add(Parts[1]);
+						}
+						else
+						{
+							if (String.Compare(Parts[0], PlatformSetupEnvVar) == 0)
+							{
+								bNeedsToWriteAutoSetupEnvVar = false;
+							}
+							// convenience for setup.bat writers.  Trim any accidental whitespace from variable names/values.
+							EnvVarNames.Add(Parts[0].Trim());
+							EnvVarValues.Add(Parts[1].Trim());
+						}
+					}
+
+					// don't actually set anything until we successfully validate and read all values in.
+					// we don't want to set a few vars, return a failure, and then have a platform try to
+					// build against a manually installed SDK with half-set env vars.
+					for (int i = 0; i < EnvVarNames.Count; ++i)
+					{
+						string EnvVarName = EnvVarNames[i];
+						string EnvVarValue = EnvVarValues[i];
+						Log.TraceVerbose("Setting variable '{0}' to '{1}'", EnvVarName, EnvVarValue);
+						Environment.SetEnvironmentVariable(EnvVarName, EnvVarValue);
+					}
+
+
+                    // actually perform the PATH stripping / adding.
+                    String OrigPathVar = Environment.GetEnvironmentVariable("PATH");
+                    String PathDelimiter = IsWindows() ? ";" : ":";
+                    String[] PathVars = { };
+                    if (!String.IsNullOrEmpty(OrigPathVar))
+                    {
+                        PathVars = OrigPathVar.Split(PathDelimiter.ToCharArray());
+                    }
+                    else
+                    {
+                        Log.TraceVerbose("Path environment variable is null during AutoSDK");
+                    }
+
+					List<String> ModifiedPathVars = new List<string>();
+					ModifiedPathVars.AddRange(PathVars);
+
+					// perform removes first, in case they overlap with any adds.
+					foreach (String PathRemove in PathRemoves)
+					{
+						foreach (String PathVar in PathVars)
+						{
+							if (PathVar.IndexOf(PathRemove, StringComparison.OrdinalIgnoreCase) >= 0)
+							{
+								Log.TraceVerbose("Removing Path: '{0}'", PathVar);
+								ModifiedPathVars.Remove(PathVar);
+							}
+						}
+					}
+
+					// remove all the of ADDs so that if this function is executed multiple times, the paths will be guaranteed to be in the same order after each run.
+					// If we did not do this, a 'remove' that matched some, but not all, of our 'adds' would cause the order to change.
+					foreach (String PathAdd in PathAdds)
+					{
+						foreach (String PathVar in PathVars)
+						{
+							if (String.Compare(PathAdd, PathVar, true) == 0)
+							{
+								Log.TraceVerbose("Removing Path: '{0}'", PathVar);
+								ModifiedPathVars.Remove(PathVar);
+							}
+						}
+					}
+
+					// perform adds, but don't add duplicates
+					foreach (String PathAdd in PathAdds)
+					{
+						if (!ModifiedPathVars.Contains(PathAdd))
+						{
+							Log.TraceVerbose("Adding Path: '{0}'", PathAdd);
+							ModifiedPathVars.Add(PathAdd);
+						}
+					}
+
+					String ModifiedPath = String.Join(PathDelimiter, ModifiedPathVars);
+					Environment.SetEnvironmentVariable("PATH", ModifiedPath);
+
+					Reader.Close();
+
+					// write out environment variable command so any process using this commandfile will mark itself as having had autosdks set up.
+					// avoids child processes spuriously detecting manualsdks.
+					if (bNeedsToWriteAutoSetupEnvVar)
+					{
+						// write out the manual sdk version since child processes won't be able to detect manual with AutoSDK messing up env vars
+						string ValueToWrite = CachedManualSDKVersion != null ? CachedManualSDKVersion : "__None";
+						using (StreamWriter Writer = File.AppendText(EnvVarFile))
+						{
+							Writer.WriteLine("{0}={1}", PlatformSetupEnvVar, ValueToWrite);
+						}
+						// set the variable in the local environment in case this process spawns any others.
+						Environment.SetEnvironmentVariable(PlatformSetupEnvVar, ValueToWrite);
+					}
+
+					// make sure we know that we've modified the local environment, invalidating manual installs for this run.
+					bLocalProcessSetupAutoSDK = true;
+
+					return true;
+				}
+			}
+			else
+			{
+				Log.TraceLog("Cannot set up environment for {1} because command file {2} does not exist.", PlatformSDKRoot, EnvVarFile);
+			}
+
+			return false;
+		}
+
+		protected void InvalidateCurrentlyInstalledAutoSDK()
+		{
+			String PlatformSDKRoot = GetPathToPlatformAutoSDKs();
+			if (Directory.Exists(PlatformSDKRoot))
+			{
+				string SDKFilename = Path.Combine(PlatformSDKRoot, CurrentlyInstalledSDKStringManifest);
+				if (File.Exists(SDKFilename))
+				{
+					File.Delete(SDKFilename);
+				}
+
+				string VersionFilename = Path.Combine(PlatformSDKRoot, LastRunScriptVersionManifest);
+				if (File.Exists(VersionFilename))
+				{
+					File.Delete(VersionFilename);
+				}
+
+				string EnvVarFile = Path.Combine(PlatformSDKRoot, SDKEnvironmentVarsFile);
+				if (File.Exists(EnvVarFile))
+				{
+					File.Delete(EnvVarFile);
+				}
+			}
+		}
+
+		/// <summary>
+		/// Currently installed AutoSDK is written out to a text file in a known location.
+		/// This function just compares the file's contents with the current requirements.
+		/// </summary>
+		public SDKStatus HasRequiredAutoSDKInstalled()
+		{
+			if (PlatformSupportsAutoSDKs() && HasAutoSDKSystemEnabled())
+			{
+				string AutoSDKRoot = GetPathToPlatformAutoSDKs();
+				if (AutoSDKRoot != "")
+				{
+					// check script version so script fixes can be propagated without touching every build machine's CurrentlyInstalled file manually.
+					bool bScriptVersionMatches = false;
+					string CurrentScriptVersionString;
+					if (GetLastRunScriptVersionString(AutoSDKRoot, out CurrentScriptVersionString) && CurrentScriptVersionString == GetRequiredScriptVersionString())
+					{
+						bScriptVersionMatches = true;
+					}
+
+					// check to make sure OutputEnvVars doesn't need regenerating
+					string EnvVarFile = Path.Combine(AutoSDKRoot, SDKEnvironmentVarsFile);
+					bool bEnvVarFileExists = File.Exists(EnvVarFile);
+
+					string CurrentSDKString;
+					if (bEnvVarFileExists && GetCurrentlyInstalledSDKString(AutoSDKRoot, out CurrentSDKString) && CurrentSDKString == GetRequiredSDKString() && bScriptVersionMatches)
+					{
+						return SDKStatus.Valid;
+					}
+					return SDKStatus.Invalid;
+				}
+			}
+			return SDKStatus.Invalid;
+		}
+
+		// This tracks if we have already checked the sdk installation.
+		private Int32 SDKCheckStatus = -1;
+
+		// true if we've ever overridden the process's environment with AutoSDK data.  After that, manual installs cannot be considered valid ever again.
+		private bool bLocalProcessSetupAutoSDK = false;
+
+		protected bool HasSetupAutoSDK()
+		{
+			return bLocalProcessSetupAutoSDK || HasParentProcessSetupAutoSDK(out _);
+		}
+
+		protected bool HasParentProcessSetupAutoSDK(out string OutAutoSDKSetupValue)
+		{
+			bool bParentProcessSetupAutoSDK = false;
+			String AutoSDKSetupVarName = GetPlatformAutoSDKSetupEnvVar();
+			OutAutoSDKSetupValue = Environment.GetEnvironmentVariable(AutoSDKSetupVarName);
+			if (!String.IsNullOrEmpty(OutAutoSDKSetupValue))
+			{
+				bParentProcessSetupAutoSDK = true;
+			}
+			return bParentProcessSetupAutoSDK;
+		}
+
+		public SDKStatus HasRequiredManualSDK()
+		{
+// 			if (HasSetupAutoSDK())
+// 			{
+// 				return SDKStatus.Invalid;
+// 			}
+//
+//			// manual installs are always invalid if we have modified the process's environment for AutoSDKs
+			return HasRequiredManualSDKInternal();
+		}
+
+		// for platforms with destructive AutoSDK.  Report if any manual sdk is installed that may be damaged by an autosdk.
+		protected virtual bool HasAnyManualInstall()
+		{
+			return false;
+		}
+
+		// tells us if the user has a valid manual install.
+		protected virtual SDKStatus HasRequiredManualSDKInternal()
+		{
+			string ManualSDKVersion;
+			GetInstalledVersions(out ManualSDKVersion, out _);
+
+			return IsVersionValid(ManualSDKVersion, bForAutoSDK:false) ? SDKStatus.Valid : SDKStatus.Invalid;
+		}
+
+		// some platforms will fail if there is a manual install that is the WRONG manual install.
+		protected virtual bool AllowInvalidManualInstall()
+		{
+			return true;
+		}
+
+		// platforms can choose if they prefer a correct the the AutoSDK install over the manual install.
+		protected virtual bool PreferAutoSDK()
+		{
+			return true;
+		}
+
+		// some platforms don't support parallel SDK installs.  AutoSDK on these platforms will
+		// actively damage an existing manual install by overwriting files in it.  AutoSDK must NOT
+		// run any setup if a manual install exists in this case.
+		protected virtual bool IsAutoSDKDestructive()
+		{
+			return false;
+		}
+
+		/// <summary>
+		/// Runs batch files if necessary to set up required AutoSDK.
+		/// AutoSDKs are SDKs that have not been setup through a formal installer, but rather come from
+		/// a source control directory, or other local copy.
+		/// </summary>
+		private void SetupAutoSDK()
+		{
+			if (IsAutoSDKSafe() && PlatformSupportsAutoSDKs() && HasAutoSDKSystemEnabled())
+			{
+				// run installation for autosdk if necessary.
+				if (HasRequiredAutoSDKInstalled() == SDKStatus.Invalid)
+				{
+					//reset check status so any checking sdk status after the attempted setup will do a real check again.
+					SDKCheckStatus = -1;
+
+					string AutoSDKRoot = GetPathToPlatformAutoSDKs();
+					string CurrentSDKString;
+					GetCurrentlyInstalledSDKString(AutoSDKRoot, out CurrentSDKString);
+
+					// switch over (note that version string can be empty)
+					if (!RunAutoSDKHooks(AutoSDKRoot, CurrentSDKString, SDKHookType.Uninstall))
+					{
+						Log.TraceLog("Failed to uninstall currently installed SDK {0}", CurrentSDKString);
+						InvalidateCurrentlyInstalledAutoSDK();
+						return;
+					}
+					// delete Manifest file to avoid multiple uninstalls
+					InvalidateCurrentlyInstalledAutoSDK();
+
+					if (!RunAutoSDKHooks(AutoSDKRoot, GetRequiredSDKString(), SDKHookType.Install, false))
+					{
+						Log.TraceLog("Failed to install required SDK {0}.  Attemping to uninstall", GetRequiredSDKString());
+						RunAutoSDKHooks(AutoSDKRoot, GetRequiredSDKString(), SDKHookType.Uninstall, false);
+						return;
+					}
+
+					string EnvVarFile = Path.Combine(AutoSDKRoot, SDKEnvironmentVarsFile);
+					if (!File.Exists(EnvVarFile))
+					{
+						Log.TraceLog("Installation of required SDK {0}.  Did not generate Environment file {1}", GetRequiredSDKString(), EnvVarFile);
+						RunAutoSDKHooks(AutoSDKRoot, GetRequiredSDKString(), SDKHookType.Uninstall, false);
+						return;
+					}
+
+					SetCurrentlyInstalledAutoSDKString(GetRequiredSDKString());
+					SetLastRunAutoSDKScriptVersion(GetRequiredScriptVersionString());
+				}
+
+				// fixup process environment to match autosdk
+				SetupEnvironmentFromAutoSDK();
+			}
+		}
+
+		#endregion
+
+		#region public AutoSDKs Utility
+
+		/// <summary>
+		/// Enum describing types of hooks a platform SDK can have
+		/// </summary>
+		public enum SDKHookType
+		{
+			Install,
+			Uninstall
+		};
+
+		/* Whether or not we should try to automatically switch SDKs when asked to validate the platform's SDK state. */
+		public static bool bAllowAutoSDKSwitching = true;
+
+		public SDKStatus SetupEnvironmentFromAutoSDK()
+		{
+			string PlatformSDKRoot = GetPathToPlatformAutoSDKs();
+
+			// load environment variables from current SDK
+			if (!SetupEnvironmentFromAutoSDK(PlatformSDKRoot))
+			{
+				Log.TraceLog("Failed to load environment from required SDK {0}", GetRequiredSDKString());
+				InvalidateCurrentlyInstalledAutoSDK();
+				return SDKStatus.Invalid;
+			}
+			return SDKStatus.Valid;
+		}
+
+		/// <summary>
+		/// Whether the required external SDKs are installed for this platform.
+		/// Could be either a manual install or an AutoSDK.
+		/// </summary>
+		public SDKStatus HasRequiredSDKsInstalled()
+		{
+			// avoid redundant potentially expensive SDK checks.
+			if (SDKCheckStatus == -1)
+			{
+				bool bHasManualSDK = HasRequiredManualSDK() == SDKStatus.Valid;
+				bool bHasAutoSDK = HasRequiredAutoSDKInstalled() == SDKStatus.Valid;
+
+				// Per-Platform implementations can choose how to handle non-Auto SDK detection / handling.
+				SDKCheckStatus = (bHasManualSDK || bHasAutoSDK) ? 1 : 0;
+			}
+			return SDKCheckStatus == 1 ? SDKStatus.Valid : SDKStatus.Invalid;
+		}
+
+		// Arbitrates between manual SDKs and setting up AutoSDK based on program options and platform preferences.
+		public void ManageAndValidateSDK()
+		{
+			// do not modify installed manifests if parent process has already set everything up.
+			// this avoids problems with determining IsAutoSDKSafe and doing an incorrect invalidate.
+			if (bAllowAutoSDKSwitching && !HasParentProcessSetupAutoSDK(out _))
+			{
+				bool bSetSomeSDK = false;
+				bool bHasRequiredManualSDK = HasRequiredManualSDK() == SDKStatus.Valid;
+				if (IsAutoSDKSafe() && (PreferAutoSDK() || !bHasRequiredManualSDK))
+				{
+					SetupAutoSDK();
+					bSetSomeSDK = true;
+				}
+
+				//Setup manual SDK if autoSDK setup was skipped or failed for whatever reason.
+				if (bHasRequiredManualSDK && (HasRequiredAutoSDKInstalled() != SDKStatus.Valid))
+				{
+					SetupManualSDK();
+					bSetSomeSDK = true;
+				}
+
+				if (!bSetSomeSDK)
+				{
+					InvalidateCurrentlyInstalledAutoSDK();
+				}
+			}
+
+			// print all SDKs to log file
+			PrintSDKInfo(LogEventType.Log, LogFormatOptions.NoConsoleOutput);
+		}
+		
+		public virtual void PrintSDKInfo(LogEventType Verbosity = LogEventType.Console, LogFormatOptions Options = LogFormatOptions.None)
+		{
+			string ManualSDKVersion, AutoSDKVersion;
+			GetInstalledVersions(out ManualSDKVersion, out AutoSDKVersion);
+
+
+			if (HasSetupAutoSDK())
+			{
+				string PlatformSDKRoot = GetPathToPlatformAutoSDKs();
+
+				UInt64 Ver;
+				TryConvertVersionToInt(AutoSDKVersion, out Ver);
+				Log.WriteLine(Verbosity, Options, "{0} using Auto SDK from: {1} 0x{2:X}", PlatformName, Path.Combine(PlatformSDKRoot, AutoSDKVersion), Ver);
+			}
+			else if (HasRequiredManualSDK() == SDKStatus.Valid)
+			{
+				Log.WriteLine(Verbosity, Options, "{0} using Manual SDK {1}", PlatformName, ManualSDKVersion);
+			}
+			else
+			{
+				string MinVersionString, MaxVersionString;
+				GetValidVersionRange(out MinVersionString, out MaxVersionString);
+
+				StringBuilder Msg = new StringBuilder();
+				Msg.AppendLine("Unable to find a valid SDK for {0}", PlatformName);
+				if (ManualSDKVersion != null)
+				{
+					Msg.AppendLine("  Found Version: {0}", ManualSDKVersion);
+				}
+
+				if (MinVersionString != MaxVersionString)
+				{
+					Msg.AppendLine("  SDK version needs to be between {0} and {1}", MinVersionString, MaxVersionString);
+				}
+				else
+				{
+					Msg.AppendLine("  SDK version needs to be {0}", MinVersionString);
+				}
+
+				Msg.AppendLine("  If your Studio has it set up, you can run this command to find the SDK to install:");
+				Msg.AppendLine("    RunUAT Turnkey -command=InstallSdk -platform={0}", PlatformName);
+
+				Log.WriteLine(LogEventType.Error, Options, Msg.ToString());
+			}
+		}
+
+		#endregion
+	}
+}
diff --git a/Engine/Source/Programs/UnrealBuildTool/Configuration/UEBuildPlatform.cs b/Engine/Source/Programs/UnrealBuildTool/Configuration/UEBuildPlatform.cs
--- a/Engine/Source/Programs/UnrealBuildTool/Configuration/UEBuildPlatform.cs
+++ b/Engine/Source/Programs/UnrealBuildTool/Configuration/UEBuildPlatform.cs
@@ -42,9 +42,13 @@ namespace UnrealBuildTool
 		/// Constructor.
 		/// </summary>
 		/// <param name="InPlatform">The enum value for this platform</param>
-		public UEBuildPlatform(UnrealTargetPlatform InPlatform)
+		/// <param name="SDK">The SDK management object for this platform</param>
+		public UEBuildPlatform(UnrealTargetPlatform InPlatform, UEBuildPlatformSDK SDK)
 		{
 			Platform = InPlatform;
+			UEBuildPlatformSDK.RegisterSDKForPlatform(SDK, Platform.ToString());
+
+			SDK.ManageAndValidateSDK();
 		}
 
 		/// <summary>
@@ -153,7 +157,10 @@ namespace UnrealBuildTool
 		/// <summary>
 		/// Whether the required external SDKs are installed for this platform. Could be either a manual install or an AutoSDK.
 		/// </summary>
-		public abstract SDKStatus HasRequiredSDKsInstalled();
+		public SDKStatus HasRequiredSDKsInstalled()
+		{
+			return UEBuildPlatform.GetSDK(Platform).HasRequiredSDKsInstalled();
+		}
 
 		/// <summary>
 		/// Whether this platform requires specific Visual Studio version.
@@ -800,6 +807,23 @@ namespace UnrealBuildTool
 			}
 		}
 
+		/// <summary>
+		/// Gets the SDK object that was passed in to the constructor
+		/// </summary>
+		/// <returns>The SDK object</returns>
+		public UEBuildPlatformSDK GetSDK()
+		{
+			return UEBuildPlatformSDK.GetSDKForPlatform(Platform.ToString());
+		}
+		/// <summary>
+		/// Gets the SDK object that was passed in to the constructor to the UEBuildPlatform constructor for this platform
+		/// </summary>
+		/// <returns>The SDK object</returns>
+		public static UEBuildPlatformSDK GetSDK(UnrealTargetPlatform Platform)
+		{
+			return UEBuildPlatformSDK.GetSDKForPlatform(Platform.ToString());
+		}
+
 		/// <summary>
 		/// Whether this platform should create debug information or not
 		/// </summary>
diff --git a/Engine/Source/Programs/UnrealBuildTool/Configuration/UEBuildPlatformSDK.cs b/Engine/Source/Programs/UnrealBuildTool/Configuration/UEBuildPlatformSDK.cs
deleted file mode 100644
diff --git a/Engine/Source/Programs/UnrealBuildTool/Modes/BuildMode.cs b/Engine/Source/Programs/UnrealBuildTool/Modes/BuildMode.cs
--- a/Engine/Source/Programs/UnrealBuildTool/Modes/BuildMode.cs
+++ b/Engine/Source/Programs/UnrealBuildTool/Modes/BuildMode.cs
@@ -202,6 +202,9 @@ namespace UnrealBuildTool
 							FileMetadataPrefetch.QueueProjectDirectory(ProjectDirectory);
 							ProjectDirs.Add(ProjectDirectory);
 						}
+
+						// print out SDK info for only the platforms that are being compiled
+						UEBuildPlatformSDK.GetSDKForPlatform(TargetDesc.Platform.ToString()).PrintSDKInfo();
 					}
 
 					// Get all the build options
diff --git a/Engine/Source/Programs/UnrealBuildTool/Platform/Android/AndroidToolChain.cs b/Engine/Source/Programs/UnrealBuildTool/Platform/Android/AndroidToolChain.cs
--- a/Engine/Source/Programs/UnrealBuildTool/Platform/Android/AndroidToolChain.cs
+++ b/Engine/Source/Programs/UnrealBuildTool/Platform/Android/AndroidToolChain.cs
@@ -13,11 +13,6 @@ namespace UnrealBuildTool
 {
 	class AndroidToolChain : ISPCToolChain, IAndroidToolChain
 	{
-		// Android NDK toolchain that must be used for C++ compiling
-		readonly int MinimumNDKToolchain = 210100;
-		readonly int MaximumNDKToolchain = 230100;
-		readonly int RecommendedNDKToolchain = 210200;
-
 		public static readonly string[] AllCpuSuffixes =
 		{
 			"-armv7",
@@ -94,8 +89,7 @@ namespace UnrealBuildTool
 		};
 
 		public string NDKToolchainVersion;
-		public string NDKDefine;
-		public int NDKDefineInt;
+		public UInt64 NDKVersionInt;
 
 		protected void SetClangVersion(int Major, int Minor, int Patch)
 		{
@@ -109,14 +103,6 @@ namespace UnrealBuildTool
 			return string.Format("{0}.{1}.{2}", ClangVersionMajor, ClangVersionMinor, ClangVersionPatch);
 		}
 
-		public void ShowNDKWarnings()
-		{
-			if (NDKDefineInt < RecommendedNDKToolchain)
-			{
-				Log.TraceInformation("Note: Android toolchain NDK " + ToolchainIntToString(RecommendedNDKToolchain) + " recommended");
-			}
-		}
-
 		/// <summary>
 		/// Checks if compiler version matches the requirements
 		/// </summary>
@@ -137,17 +123,6 @@ namespace UnrealBuildTool
 				(ClangVersionMajor == Major && ClangVersionMinor == Minor && ClangVersionPatch < Patch);
 		}
 
-		private static string ToolchainIntToString(int ToolchainInt)
-		{
-			int RevisionNum = ToolchainInt / 10000;
-			int RevisionMinor = ToolchainInt - (RevisionNum * 10000);
-			int RevisionLetterNum = RevisionMinor / 100;
-			//int RevisionBeta = RevisionMinor - (RevisionLetterNum * 100);
-			char RevisionLetter = Convert.ToChar('a' + RevisionLetterNum - 1);
-
-			return "r" + RevisionNum + (RevisionLetterNum > 1 ? Char.ToString(RevisionLetter) : "");
-		}
-
 		[CommandLine("-Architectures=", ListSeparator = '+')]
 		public List<string> ArchitectureArg = new List<string>();
 
@@ -198,49 +173,7 @@ namespace UnrealBuildTool
 
 			NDKPath = NDKPath.Replace("\"", "");
 
-			// figure out the NDK version
-			NDKToolchainVersion = "unknown";
-			NDKDefine = "100500";    // assume r10e
-			string SourcePropFilename = Path.Combine(NDKPath, "source.properties");
-			if (File.Exists(SourcePropFilename))
-			{
-				string RevisionString = "";
-				string[] PropertyContents = File.ReadAllLines(SourcePropFilename);
-				foreach (string PropertyLine in PropertyContents)
-				{
-					if (PropertyLine.StartsWith("Pkg.Revision"))
-					{
-						RevisionString = PropertyLine;
-						break;
-					}
-				}
-
-				int EqualsIndex = RevisionString.IndexOf('=');
-				if (EqualsIndex > 0)
-				{
-					string[] RevisionParts = RevisionString.Substring(EqualsIndex + 1).Trim().Split('.');
-					int RevisionMinor = int.Parse(RevisionParts.Length > 1 ? RevisionParts[1] : "0");
-					char RevisionLetter = Convert.ToChar('a' + RevisionMinor);
-					int RevisionBeta = 0;  // @TODO
-					NDKToolchainVersion = "r" + RevisionParts[0] + (RevisionMinor > 0 ? Char.ToString(RevisionLetter) : "");
-					NDKDefine = RevisionParts[0] + string.Format("{0:00}", RevisionMinor + 1) + string.Format("{0:00}", RevisionBeta);
-				}
-			}
-			else
-			{
-				string ReleaseFilename = Path.Combine(NDKPath, "RELEASE.TXT");
-				if (File.Exists(ReleaseFilename))
-				{
-					string[] PropertyContents = File.ReadAllLines(SourcePropFilename);
-					NDKToolchainVersion = PropertyContents[0];
-				}
-			}
-			if (!int.TryParse(NDKDefine, out NDKDefineInt))
-			{
-				NDKDefineInt = 100500;
-			}
-
-			string ArchitecturePath = "";
+			string ArchitecturePath;
 			string ArchitecturePathWindows32 = @"prebuilt/windows";
 			string ArchitecturePathWindows64 = @"prebuilt/windows-x86_64";
 			string ArchitecturePathMac = @"prebuilt/darwin-x86_64";
@@ -274,6 +207,11 @@ namespace UnrealBuildTool
 				throw new BuildException("Couldn't find 32-bit or 64-bit versions of the Android toolchain with NDKROOT: " + NDKPath);
 			}
 
+			// get the installed version (in the form r10e and 100500)
+			UEBuildPlatformSDK SDK = UEBuildPlatform.GetSDK(UnrealTargetPlatform.Android);
+			NDKToolchainVersion = SDK.GetInstalledVersion();
+			SDK.TryConvertVersionToInt(NDKToolchainVersion, out NDKVersionInt);
+
 			// figure out clang version (will live in toolchains/llvm from NDK 21 forward
 			if (Directory.Exists(Path.Combine(NDKPath, "toolchains", "llvm")))
 			{
@@ -292,17 +230,9 @@ namespace UnrealBuildTool
 				throw new BuildException("Cannot find supported Android toolchain with NDKPath:" + NDKPath);
 			}
 
-			// verify NDK toolchain is supported
-			if ((NDKDefineInt < MinimumNDKToolchain || NDKDefineInt > MaximumNDKToolchain)
-				&& !bAllowMissingNDK)
-			{
-				throw new BuildException("Android toolchain NDK " + ToolchainIntToString(NDKDefineInt) + " not supported; please use NDK " + ToolchainIntToString(MinimumNDKToolchain) + " to NDK " + ToolchainIntToString(MaximumNDKToolchain) +
-					" (NDK " + ToolchainIntToString(RecommendedNDKToolchain) + " recommended)");
-			}
-
 			// set up the path to our toolchains
 			ClangPath = Utils.CollapseRelativeDirectories(Path.Combine(NDKPath, @"toolchains/llvm", ArchitecturePath, @"bin/clang++" + ExeExtension));
-			if (NDKDefineInt < 210000 || ForceLDLinker())
+			if (ForceLDLinker())
 			{
 				// use ld before r21
 				ArPathArm = Utils.CollapseRelativeDirectories(Path.Combine(NDKPath, @"toolchains/arm-linux-androideabi-4.9", ArchitecturePath, @"bin/armv7a-linux-androideabi-ar" + ExeExtension));
@@ -891,7 +821,7 @@ namespace UnrealBuildTool
 				Result += " -Wl,--strip-debug";
 			}
 
-			bool bUseLLD = NDKDefineInt >= 210000 && !ForceLDLinker();
+			bool bUseLLD = !ForceLDLinker();
 			bool bAllowLdGold = true;
 			if (Architecture == "-arm64")
 			{
diff --git a/Engine/Source/Programs/UnrealBuildTool/Platform/Android/UEBuildAndroid.cs b/Engine/Source/Programs/UnrealBuildTool/Platform/Android/UEBuildAndroid.cs
--- a/Engine/Source/Programs/UnrealBuildTool/Platform/Android/UEBuildAndroid.cs
+++ b/Engine/Source/Programs/UnrealBuildTool/Platform/Android/UEBuildAndroid.cs
@@ -8,6 +8,7 @@ using System.IO;
 using System.Xml;
 using System.Linq;
 using Tools.DotNETCommon;
+using System.Text.RegularExpressions;
 
 namespace UnrealBuildTool
 {
@@ -74,10 +75,10 @@ namespace UnrealBuildTool
 
 	class AndroidPlatform : UEBuildPlatform
 	{
-		AndroidPlatformSDK SDK;
+		UEBuildPlatformSDK SDK;
 
-		public AndroidPlatform(UnrealTargetPlatform InTargetPlatform, AndroidPlatformSDK InSDK) 
-			: base(InTargetPlatform)
+		public AndroidPlatform(UnrealTargetPlatform InTargetPlatform, UEBuildPlatformSDK InSDK) 
+			: base(InTargetPlatform, InSDK)
 		{
 			SDK = InSDK;
 		}
@@ -86,11 +87,6 @@ namespace UnrealBuildTool
 		{
 		}
 
-		public override SDKStatus HasRequiredSDKsInstalled()
-		{
-			return SDK.HasRequiredSDKsInstalled();
-		}
-
 		public override void ResetTarget(TargetRules Target)
 		{
 			ValidateTarget(Target);
@@ -327,18 +323,14 @@ namespace UnrealBuildTool
 			AndroidToolChain ToolChain = new AndroidToolChain(Target.ProjectFile, false, Target.AndroidPlatform.Architectures, Target.AndroidPlatform.GPUArchitectures);
 
 			// figure out the NDK version
-			string NDKToolchainVersion = ToolChain.NDKToolchainVersion;
-			string NDKDefine = ToolChain.NDKDefine;
-
-			string GccVersion = "4.9";
+			string NDKToolchainVersion = SDK.GetInstalledVersion();
+			UInt64 NDKVersionInt;
+			SDK.TryConvertVersionToInt(NDKToolchainVersion, out NDKVersionInt);
 
 			// PLATFORM_ANDROID_NDK_VERSION is in the form 150100, where 15 is major version, 01 is the letter (1 is 'a'), 00 indicates beta revision if letter is 00
-			Log.TraceInformation("PLATFORM_ANDROID_NDK_VERSION = {0}", NDKDefine);
-			CompileEnvironment.Definitions.Add("PLATFORM_ANDROID_NDK_VERSION=" + NDKDefine);
+			CompileEnvironment.Definitions.Add(string.Format("PLATFORM_ANDROID_NDK_VERSION={0}", NDKVersionInt));
 
-			int NDKVersionInt = ToolChain.GetNdkApiLevelInt();
-			Log.TraceInformation("NDK toolchain: {0}, NDK version: {1}, GccVersion: {2}, ClangVersion: {3}", NDKToolchainVersion, NDKVersionInt.ToString(), GccVersion, ToolChain.GetClangVersionString());
-			ToolChain.ShowNDKWarnings();
+			Log.TraceInformation("NDK toolchain: {0}, NDK version: {1}, ClangVersion: {2}", NDKToolchainVersion, NDKVersionInt, ToolChain.GetClangVersionString());
 
 			CompileEnvironment.Definitions.Add("PLATFORM_DESKTOP=0");
 			CompileEnvironment.Definitions.Add("PLATFORM_CAN_SUPPORT_EDITORONLY_DATA=0");
@@ -453,19 +445,114 @@ namespace UnrealBuildTool
 
 	class AndroidPlatformSDK : UEBuildPlatformSDK
 	{
-		protected override bool PlatformSupportsAutoSDKs()
+		public override string GetDesiredVersion()
 		{
-			return true;
+			return "-23";
 		}
 
-		public override string GetSDKTargetPlatformName()
+ 		public override void GetValidVersionRange(out string MinVersion, out string MaxVersion)
+ 		{
+			MinVersion = "r21a";
+ 			MaxVersion = "r23a";
+ 		}
+
+		// the version we prefer
+		private string RecommendedVersion = "21b";
+
+		public override string GetInstalledSDKVersion()
 		{
-			return "Android";
+			string NDKPath = Environment.GetEnvironmentVariable("NDKROOT");
+
+			// don't register if we don't have an NDKROOT specified
+			if (String.IsNullOrEmpty(NDKPath))
+			{
+				return null;
+			}
+
+			NDKPath = NDKPath.Replace("\"", "");
+
+			// figure out the NDK version
+			string NDKToolchainVersion = null;
+			string SourcePropFilename = Path.Combine(NDKPath, "source.properties");
+			if (File.Exists(SourcePropFilename))
+			{
+				string RevisionString = "";
+				string[] PropertyContents = File.ReadAllLines(SourcePropFilename);
+				foreach (string PropertyLine in PropertyContents)
+				{
+					if (PropertyLine.StartsWith("Pkg.Revision"))
+					{
+						RevisionString = PropertyLine;
+						break;
+					}
+				}
+
+				int EqualsIndex = RevisionString.IndexOf('=');
+				if (EqualsIndex > 0)
+				{
+					string[] RevisionParts = RevisionString.Substring(EqualsIndex + 1).Trim().Split('.');
+					int RevisionMinor = int.Parse(RevisionParts.Length > 1 ? RevisionParts[1] : "0");
+					char RevisionLetter = Convert.ToChar('a' + RevisionMinor);
+					NDKToolchainVersion = "r" + RevisionParts[0] + (RevisionMinor > 0 ? Char.ToString(RevisionLetter) : "");
+				}
+			}
+			else
+			{
+				string ReleaseFilename = Path.Combine(NDKPath, "RELEASE.TXT");
+				if (File.Exists(ReleaseFilename))
+				{
+					string[] PropertyContents = File.ReadAllLines(SourcePropFilename);
+					NDKToolchainVersion = PropertyContents[0];
+				}
+			}
+
+			// return something like r10e, or null if anything went wrong above
+			return NDKToolchainVersion;
+
+		}
+
+
+		public override bool TryConvertVersionToInt(string StringValue, out UInt64 OutValue)
+		{
+			// convert r<num>[letter] to hex
+			Match Result = Regex.Match(StringValue, @"^r(\d*)([a-z])?");
+
+			if (Result.Success)
+			{
+				// 8 for number, 8 for letter, 8 for unused patch
+				int RevisionNumber = 0;
+				if (Result.Groups[2].Success)
+				{
+					RevisionNumber = (Result.Groups[2].Value[0] - 'a') + 1;
+				}
+				string VersionString = string.Format("{0}{1:00}{2:00}", Result.Groups[1], RevisionNumber, 0);
+				return UInt64.TryParse(VersionString, out OutValue);
+			}
+
+			OutValue = 0;
+			return false;
+		}
+
+
+		public override void PrintSDKInfo(LogEventType Verbosity, LogFormatOptions Options)
+		{
+			base.PrintSDKInfo(Verbosity, Options);
+
+			if (GetInstalledVersion() != RecommendedVersion)
+			{
+				Log.WriteLine(Verbosity, Options, "Note: Android toolchain NDK {0} recommended", RecommendedVersion);
+			}
+		}
+
+
+		protected override bool PlatformSupportsAutoSDKs()
+		{
+			return true;
 		}
 
 		protected override string GetRequiredSDKString()
 		{
-			return "-23";
+			return GetDesiredVersion();
 		}
 
 		protected override String GetRequiredScriptVersionString()
@@ -632,7 +719,6 @@ namespace UnrealBuildTool
 		public override void RegisterBuildPlatforms()
 		{
 			AndroidPlatformSDK SDK = new AndroidPlatformSDK();
-			SDK.ManageAndValidateSDK();
 
 			// Register this build platform
 			UEBuildPlatform.RegisterBuildPlatform(new AndroidPlatform(SDK));
diff --git a/Engine/Source/Programs/UnrealBuildTool/Platform/HoloLens/UEBuildHoloLens.cs b/Engine/Source/Programs/UnrealBuildTool/Platform/HoloLens/UEBuildHoloLens.cs
--- a/Engine/Source/Programs/UnrealBuildTool/Platform/HoloLens/UEBuildHoloLens.cs
+++ b/Engine/Source/Programs/UnrealBuildTool/Platform/HoloLens/UEBuildHoloLens.cs
@@ -8,6 +8,7 @@ using System.IO;
 using System.Linq;
 using Microsoft.Win32;
 using Tools.DotNETCommon;
+using System.Text.RegularExpressions;
 
 namespace UnrealBuildTool
 {
@@ -140,16 +141,9 @@ namespace UnrealBuildTool
 		public static readonly Version MaximumSDKVersionForVS2015 = new Version(10, 0, 14393, int.MaxValue);
 		public static readonly Version MinimumSDKVersionForD3D12RHI = new Version(10, 0, 15063, 0);
 
-		HoloLensPlatformSDK SDK;
 
-		public HoloLens(UnrealTargetPlatform InPlatform, HoloLensPlatformSDK InSDK) : base(InPlatform)
+		public HoloLens(UnrealTargetPlatform InPlatform, HoloLensPlatformSDK InSDK) : base(InPlatform, InSDK)
 		{
-			SDK = InSDK;
-		}
-
-		public override SDKStatus HasRequiredSDKsInstalled()
-		{
-			return SDK.HasRequiredSDKsInstalled();
 		}
 
 		public override void ValidateTarget(TargetRules Target)
@@ -749,45 +743,9 @@ namespace UnrealBuildTool
 
 
 
-	class HoloLensPlatformSDK : UEBuildPlatformSDK
+	class HoloLensPlatformSDK : MicrosoftPlatformSDK
 	{
-		static bool bIsInstalled = false;
-		//static string LatestVersionString = string.Empty;
-		//static string InstallLocation = string.Empty;
-
-		static HoloLensPlatformSDK()
-		{
-#if !__MonoCS__
-			if (Utils.IsRunningOnMono)
-			{
-				return;
-			}
-
-			string Version = "v10.0";
-			string[] possibleRegLocations =
-			{
-				@"HKEY_LOCAL_MACHINE\SOFTWARE\Wow6432Node\Microsoft\Microsoft SDKs\Windows\",
-				@"HKEY_CURRENT_USER\SOFTWARE\Wow6432Node\Microsoft\Microsoft SDKs\Windows\"
-			};
-			foreach (string regLocation in possibleRegLocations)
-			{
-				object Result = Microsoft.Win32.Registry.GetValue(@"HKEY_LOCAL_MACHINE\SOFTWARE\Wow6432Node\Microsoft\Microsoft SDKs\Windows\" + Version, "InstallationFolder", null);
 
-				if (Result != null)
-				{
-					bIsInstalled = true;
-					//InstallLocation = (string)Result;
-					//LatestVersionString = Microsoft.Win32.Registry.GetValue(@"HKEY_LOCAL_MACHINE\SOFTWARE\Wow6432Node\Microsoft\Microsoft SDKs\Windows\" + Version, "ProductVersion", null) as string;
-					break;
-				}
-			}
-#endif
-		}
-
-		protected override SDKStatus HasRequiredManualSDKInternal()
-		{
-			return (!Utils.IsRunningOnMono && bIsInstalled) ? SDKStatus.Valid : SDKStatus.Invalid;
-		}
 	}
 
 	class HoloLensPlatformFactory : UEBuildPlatformFactory
@@ -803,7 +761,6 @@ namespace UnrealBuildTool
 		public override void RegisterBuildPlatforms()
 		{
 			HoloLensPlatformSDK SDK = new HoloLensPlatformSDK();
-			SDK.ManageAndValidateSDK();
 
 			UEBuildPlatform.RegisterBuildPlatform(new HoloLens(UnrealTargetPlatform.HoloLens, SDK));
 			UEBuildPlatform.RegisterPlatformWithGroup(UnrealTargetPlatform.HoloLens, UnrealPlatformGroup.Microsoft);
diff --git a/Engine/Source/Programs/UnrealBuildTool/Platform/IOS/UEBuildIOS.cs b/Engine/Source/Programs/UnrealBuildTool/Platform/IOS/UEBuildIOS.cs
--- a/Engine/Source/Programs/UnrealBuildTool/Platform/IOS/UEBuildIOS.cs
+++ b/Engine/Source/Programs/UnrealBuildTool/Platform/IOS/UEBuildIOS.cs
@@ -6,9 +6,7 @@ using System.Text;
 using System.Diagnostics;
 using System.IO;
 using System.Linq;
-using System.Xml;
 using Tools.DotNETCommon;
-using Microsoft.Win32;
 
 namespace UnrealBuildTool
 {
@@ -648,7 +646,7 @@ namespace UnrealBuildTool
 						if (idx > 0)
 						{
 							idx += "<string>".Length;
-							String FullID = AllText.Substring(idx, AllText.IndexOf("</string>", idx) - idx);
+							string FullID = AllText.Substring(idx, AllText.IndexOf("</string>", idx) - idx);
 							BundleIdentifier = FullID.Substring(FullID.IndexOf('.') + 1);
 						}
 					}
@@ -706,7 +704,6 @@ namespace UnrealBuildTool
 
 	class IOSPlatform : UEBuildPlatform
 	{
-		IOSPlatformSDK SDK;
 		List<IOSProjectSettings> CachedProjectSettings = new List<IOSProjectSettings>();
 		List<IOSProjectSettings> CachedProjectSettingsByBundle = new List<IOSProjectSettings>();
 		Dictionary<string, IOSProvisioningData> ProvisionCache = new Dictionary<string, IOSProvisioningData>();
@@ -714,15 +711,14 @@ namespace UnrealBuildTool
 		// by default, use an empty architecture (which is really just a modifer to the platform for some paths/names)
 		public static string IOSArchitecture = "";
 
-		public IOSPlatform(IOSPlatformSDK InSDK)
+		public IOSPlatform(UEBuildPlatformSDK InSDK)
 			: this(InSDK, UnrealTargetPlatform.IOS)
 		{
 		}
 
-		protected IOSPlatform(IOSPlatformSDK InSDK, UnrealTargetPlatform TargetPlatform)
-			: base(TargetPlatform)
+		protected IOSPlatform(UEBuildPlatformSDK InSDK, UnrealTargetPlatform TargetPlatform)
+			: base(TargetPlatform, InSDK)
 		{
-			SDK = InSDK;
 		}
 
 		// The current architecture - affects everything about how UBT operates on IOS
@@ -809,11 +805,6 @@ namespace UnrealBuildTool
 			Target.bCheckSystemHeadersForModification = false;
 		}
 
-		public override SDKStatus HasRequiredSDKsInstalled()
-		{
-			return SDK.HasRequiredSDKsInstalled();
-		}
-
 		/// <summary>
 		/// Determines if the given name is a build product for a target.
 		/// </summary>
@@ -1197,61 +1188,6 @@ namespace UnrealBuildTool
 		}
 	}
 
-	class IOSPlatformSDK : UEBuildPlatformSDK
-	{
-		protected override SDKStatus HasRequiredManualSDKInternal()
-		{
-			if (!Utils.IsRunningOnMono)
-			{
-				// check to see if iTunes is installed
-				string dllPath = Registry.GetValue("HKEY_LOCAL_MACHINE\\SOFTWARE\\Wow6432Node\\Apple Inc.\\Apple Mobile Device Support\\Shared", "iTunesMobileDeviceDLL", null) as string;
-				if (String.IsNullOrEmpty(dllPath) || !File.Exists(dllPath))
-				{
-					dllPath = Registry.GetValue("HKEY_LOCAL_MACHINE\\SOFTWARE\\Wow6432Node\\Apple Inc.\\Apple Mobile Device Support\\Shared", "MobileDeviceDLL", null) as string;
-					if (String.IsNullOrEmpty(dllPath) || !File.Exists(dllPath))
-					{
-						dllPath = FindWindowsStoreITunesDLL();
-
-						if (String.IsNullOrEmpty(dllPath) || !File.Exists(dllPath))
-						{
-							return SDKStatus.Invalid;
-						}
-					}
-				}
-			}
-			return SDKStatus.Valid;
-		}
-
-		static string FindWindowsStoreITunesDLL()
-		{
-			string InstallPath = null;
-
-			string PackagesKeyName = "Software\\Classes\\Local Settings\\Software\\Microsoft\\Windows\\CurrentVersion\\AppModel\\PackageRepository\\Packages";
-
-			RegistryKey PackagesKey = Registry.LocalMachine.OpenSubKey(PackagesKeyName);
-			if (PackagesKey != null)
-			{
-				string[] PackageSubKeyNames = PackagesKey.GetSubKeyNames();
-
-				foreach (string PackageSubKeyName in PackageSubKeyNames)
-				{
-					if (PackageSubKeyName.Contains("AppleInc.iTunes") && (PackageSubKeyName.Contains("_x64") || PackageSubKeyName.Contains("_x86")))
-					{
-						string FullPackageSubKeyName = PackagesKeyName + "\\" + PackageSubKeyName;
-
-						RegistryKey iTunesKey = Registry.LocalMachine.OpenSubKey(FullPackageSubKeyName);
-						if (iTunesKey != null)
-						{
-							InstallPath = (string)iTunesKey.GetValue("Path") + "\\AMDS32\\MobileDevice.dll";
-							break;
-						}
-					}
-				}
-			}
-
-			return InstallPath;
-		}
-	}
 
 	class IOSPlatformFactory : UEBuildPlatformFactory
 	{
@@ -1265,8 +1201,7 @@ namespace UnrealBuildTool
 		/// </summary>
 		public override void RegisterBuildPlatforms()
 		{
-			IOSPlatformSDK SDK = new IOSPlatformSDK();
-			SDK.ManageAndValidateSDK();
+			ApplePlatformSDK SDK = new ApplePlatformSDK();
 
 			// Register this build platform for IOS
 			UEBuildPlatform.RegisterBuildPlatform(new IOSPlatform(SDK));
diff --git a/Engine/Source/Programs/UnrealBuildTool/Platform/Linux/LinuxCommon.cs b/Engine/Source/Programs/UnrealBuildTool/Platform/Linux/LinuxCommon.cs
--- a/Engine/Source/Programs/UnrealBuildTool/Platform/Linux/LinuxCommon.cs
+++ b/Engine/Source/Programs/UnrealBuildTool/Platform/Linux/LinuxCommon.cs
@@ -38,7 +38,7 @@ namespace UnrealBuildTool
 
 		public static string WhichClang()
 		{
-			string InternalSDKPath = LinuxPlatformSDK.GetInternalSDKPath();
+			string InternalSDKPath = ((LinuxPlatformSDK)UEBuildPlatform.GetSDK(UnrealTargetPlatform.Linux)).GetInternalSDKPath();
 			if (!String.IsNullOrEmpty(InternalSDKPath))
 			{
 				return Path.Combine(InternalSDKPath, "bin", "clang++");
diff --git a/Engine/Source/Programs/UnrealBuildTool/Platform/Linux/LinuxToolChain.cs b/Engine/Source/Programs/UnrealBuildTool/Platform/Linux/LinuxToolChain.cs
--- a/Engine/Source/Programs/UnrealBuildTool/Platform/Linux/LinuxToolChain.cs
+++ b/Engine/Source/Programs/UnrealBuildTool/Platform/Linux/LinuxToolChain.cs
@@ -1160,11 +1160,6 @@ namespace UnrealBuildTool
 		/// </summary>
 		private static List<FileItem> AllBinaries = new List<FileItem>();
 
-		/// <summary>
-		/// Tracks that information about used C++ library is only printed once
-		/// </summary>
-		private bool bHasPrintedBuildDetails = false;
-
 		/// <summary>
 		/// Checks if we actually can use LTO/PGO with this set of tools
 		/// </summary>
@@ -1263,24 +1258,6 @@ namespace UnrealBuildTool
 			Log.TraceInformation("------------------------------");
 		}
 
-		protected bool CheckSDKVersionFromFile(string VersionPath, out string ErrorMessage)
-		{
-			if (File.Exists(VersionPath))
-			{
-				StreamReader SDKVersionFile = new StreamReader(VersionPath);
-				string SDKVersionString = SDKVersionFile.ReadLine();
-				SDKVersionFile.Close();
-
-				if (SDKVersionString != null)
-				{
-					return PlatformSDK.CheckSDKCompatible(SDKVersionString, out ErrorMessage);
-				}
-			}
-
-			ErrorMessage = "Cannot use an old toolchain (missing " + PlatformSDK.SDKVersionFileName() + " file, assuming version earlier than v11)";
-			return false;
-		}
-
 		public override CPPOutput CompileCPPFiles(CppCompileEnvironment CompileEnvironment, List<FileItem> InputFiles, DirectoryReference OutputDir, string ModuleName, IActionGraphBuilder Graph)
 		{
 			string Arguments = GetCLArguments_Global(CompileEnvironment);
@@ -1288,28 +1265,6 @@ namespace UnrealBuildTool
 
 			//var BuildPlatform = UEBuildPlatform.GetBuildPlatform(CompileEnvironment.Platform);
 
-			if (!bHasPrintedBuildDetails)
-			{
-				PrintBuildDetails(CompileEnvironment);
-
-				string LinuxDependenciesPath = Path.Combine(UnrealBuildTool.EngineDirectory.FullName, "Source/ThirdParty/Linux", PlatformSDK.HaveLinuxDependenciesFile());
-				if (!File.Exists(LinuxDependenciesPath))
-				{
-					throw new BuildException("Please make sure that Engine/Source/ThirdParty/Linux is complete (re - run Setup script if using a github build)");
-				}
-
-				if (!String.IsNullOrEmpty(MultiArchRoot))
-				{
-					string ErrorMessage;
-					if (!CheckSDKVersionFromFile(Path.Combine(MultiArchRoot, PlatformSDK.SDKVersionFileName()), out ErrorMessage))
-					{
-						throw new BuildException(ErrorMessage);
-					}
-				}
-
-				bHasPrintedBuildDetails = true;
-			}
-
 			if ((CompileEnvironment.bAllowLTCG || CompileEnvironment.bPGOOptimize || CompileEnvironment.bPGOProfile) && !CanUseAdvancedLinkerFeatures(CompileEnvironment.Architecture))
 			{
 				throw new BuildException(ExplainWhyCannotUseAdvancedLinkerFeatures(CompileEnvironment.Architecture));
diff --git a/Engine/Source/Programs/UnrealBuildTool/Platform/Linux/UEBuildLinux.cs b/Engine/Source/Programs/UnrealBuildTool/Platform/Linux/UEBuildLinux.cs
--- a/Engine/Source/Programs/UnrealBuildTool/Platform/Linux/UEBuildLinux.cs
+++ b/Engine/Source/Programs/UnrealBuildTool/Platform/Linux/UEBuildLinux.cs
@@ -165,19 +165,11 @@ namespace UnrealBuildTool
 		}
 
 		public LinuxPlatform(UnrealTargetPlatform UnrealTarget, LinuxPlatformSDK InSDK)
-			: base(UnrealTarget)
+			: base(UnrealTarget, InSDK)
 		{
 			SDK = InSDK;
 		}
 
-		/// <summary>
-		/// Whether the required external SDKs are installed for this platform. Could be either a manual install or an AutoSDK.
-		/// </summary>
-		public override SDKStatus HasRequiredSDKsInstalled()
-		{
-			return SDK.HasRequiredSDKsInstalled();
-		}
-
 		/// <summary>
 		/// Find the default architecture for the given project
 		/// </summary>
@@ -619,15 +611,81 @@ namespace UnrealBuildTool
 
 	class LinuxPlatformSDK : UEBuildPlatformSDK
 	{
-		/// <summary>
-		/// This is the SDK version we support
-		/// </summary>
-		static string ExpectedSDKVersion = "v16_clang-9.0.1-centos7";	// now unified for all the architectures
+		public override string GetDesiredVersion()
+		{
+			return "v15_clang-8.0.1-centos7";
+//			return "v16_clang-9.0.1-centos7";
+		}
+
+		public override void GetValidVersionRange(out string MinVersion, out string MaxVersion)
+ 		{
+			// all that matters is the number after the v, according to TryConvertVersionToInt()
+			MinVersion = "v10_clang-5.0.0-centos7";
+ 			MaxVersion = "v16_clang-9.0.1-centos7";
+ 		}
+
+		public override string GetAutoSDKPlatformName()
+		{
+			return TargetPlatformName;
+		}
+
+
+		public override string GetInstalledSDKVersion()
+		{
+			// @todo turnkey: ForceUseSystemCompiler() returns true, we should probably run system clang -V or similar to get version
+
+			string SDKDir = GetSDKLocation();
+			if (string.IsNullOrEmpty(SDKDir))
+			{
+				return null;
+			}
+
+			string VersionFile = Path.Combine(SDKDir, SDKVersionFileName());
+
+			if (!File.Exists(VersionFile))
+			{
+				// ErrorMessage = "Cannot use an old toolchain (missing " + PlatformSDK.SDKVersionFileName() + " file, assuming version earlier than v11)";
+				return null;
+			}
+
+			StreamReader SDKVersionFile = new StreamReader(VersionFile);
+			string SDKVersionString = SDKVersionFile.ReadLine();
+			SDKVersionFile.Close();
+
+			return SDKVersionString;
+		}
+
+
+		public override bool TryConvertVersionToInt(string StringValue, out UInt64 OutValue)
+		{
+			// Example: v11_clang-5.0.0-centos7
+			string FullVersionPattern = @"^v([0-9]+)_.*$";
+			Match Result = Regex.Match(StringValue, FullVersionPattern);
+			if (Result.Success)
+			{
+				return UInt64.TryParse(Result.Groups[1].Value, out OutValue);
+			}
+
+			OutValue = 0;
+			return false;
+		}
+
+
+
+
+
+		protected override bool PlatformSupportsAutoSDKs()
+		{
+			return true;
+		}
+
+
+
 
 		/// <summary>
 		/// Platform name (embeds architecture for now)
 		/// </summary>
-		static private string TargetPlatformName = "Linux_x64";
+		private const string TargetPlatformName = "Linux_x64";
 
 		/// <summary>
 		/// Force using system compiler and error out if not possible
@@ -648,37 +706,24 @@ namespace UnrealBuildTool
 		/// Whether platform supports switching SDKs during runtime
 		/// </summary>
 		/// <returns>true if supports</returns>
-		protected override bool PlatformSupportsAutoSDKs()
-		{
-			return true;
-		}
-
-		/// <summary>
-		/// Returns platform-specific name used in SDK repository
-		/// </summary>
-		/// <returns>path to SDK Repository</returns>
-		public override string GetSDKTargetPlatformName()
-		{
-			return TargetPlatformName;
-		}
 
 		/// <summary>
 		/// Returns a path to the internal SDK
 		/// </summary>
 		/// <returns>Valid path to the internal SDK, null otherwise</returns>
-		static public string GetInternalSDKPath()
+		public string GetInternalSDKPath()
 		{
 			string SDKRoot = Environment.GetEnvironmentVariable(SDKRootEnvVar);
 			if (!String.IsNullOrEmpty(SDKRoot))
 			{
-				string AutoSDKPath = Path.Combine(SDKRoot, "Host" + BuildHostPlatform.Current.Platform, TargetPlatformName, ExpectedSDKVersion, LinuxPlatform.DefaultHostArchitecture);
+				string AutoSDKPath = Path.Combine(SDKRoot, "Host" + BuildHostPlatform.Current.Platform, TargetPlatformName, GetDesiredVersion(), LinuxPlatform.DefaultHostArchitecture);
 				if (DirectoryReference.Exists(new DirectoryReference(AutoSDKPath)))
 				{
 					return AutoSDKPath;
 				}
 			}
 
-			string InTreeSDKPath = Path.Combine(LinuxPlatformSDK.GetInTreeSDKRoot().FullName, ExpectedSDKVersion, LinuxPlatform.DefaultHostArchitecture);
+			string InTreeSDKPath = Path.Combine(LinuxPlatformSDK.GetInTreeSDKRoot().FullName, GetDesiredVersion(), LinuxPlatform.DefaultHostArchitecture);
 			if (DirectoryReference.Exists(new DirectoryReference(InTreeSDKPath)))
 			{
 				return InTreeSDKPath;
@@ -693,7 +738,7 @@ namespace UnrealBuildTool
 		/// <returns>Valid SDK string</returns>
 		protected override string GetRequiredSDKString()
 		{
-			return ExpectedSDKVersion;
+			return GetDesiredVersion();
 		}
 
 		protected override String GetRequiredScriptVersionString()
@@ -718,56 +763,6 @@ namespace UnrealBuildTool
 			return "ToolchainVersion.txt";
 		}
 
-		protected static int GetLinuxToolchainVersionFromString(string SDKVersion)
-		{
-			// Example: v11_clang-5.0.0-centos7
-			string FullVersionPattern = @"^v[0-9]+_.*$";
-			Regex Regex = new Regex(FullVersionPattern);
-			if (Regex.IsMatch(SDKVersion))
-			{
-				string VersionPattern = @"[0-9]+";
-				Regex = new Regex(VersionPattern);
-				Match Match = Regex.Match(SDKVersion);
-				if (Match.Success)
-				{
-					int Version;
-					bool bParsed = Int32.TryParse(Match.Value, out Version);
-					if (bParsed)
-					{
-						return Version;
-					}
-				}
-			}
-
-			return -1;
-		}
-
-		public bool CheckSDKCompatible(string VersionString, out string ErrorMessage)
-		{
-			int Version = GetLinuxToolchainVersionFromString(VersionString);
-			int ExpectedVersion = GetLinuxToolchainVersionFromString(ExpectedSDKVersion);
-			if (Version >= 0 && ExpectedVersion >= 0 && Version != ExpectedVersion)
-			{
-				if (Version < ExpectedVersion)
-				{
-					ErrorMessage = "Toolchain found \"" + VersionString + "\" is older then the required version \"" + ExpectedSDKVersion + "\"";
-					return false;
-				}
-				else
-				{
-					Log.TraceWarning("Toolchain \"{0}\" is newer than the expected version \"{1}\", you may run into compilation errors", VersionString, ExpectedSDKVersion);
-				}
-			}
-			else if (VersionString != ExpectedSDKVersion)
-			{
-				ErrorMessage = "Failed to find a supported toolchain, found \"" + VersionString + "\", expected \"" + ExpectedSDKVersion + "\"";
-				return false;
-			}
-
-			ErrorMessage = "";
-			return true;
-		}
-
 		/// <summary>
 		/// Returns the in-tree root for the Linux Toolchain for this host platform.
 		/// </summary>
@@ -816,7 +811,7 @@ namespace UnrealBuildTool
 				DirectoryReference InTreeSDKVersionRoot = GetInTreeSDKRoot();
 				if (InTreeSDKVersionRoot != null)
 				{
-					DirectoryReference InTreeSDKVersionPath = DirectoryReference.Combine(InTreeSDKVersionRoot, ExpectedSDKVersion);
+					DirectoryReference InTreeSDKVersionPath = DirectoryReference.Combine(InTreeSDKVersionRoot, GetDesiredVersion());
 					if (DirectoryReference.Exists(InTreeSDKVersionPath))
 					{
 						MultiArchRoot = InTreeSDKVersionPath.FullName;
@@ -903,7 +898,7 @@ namespace UnrealBuildTool
 		public override void RegisterBuildPlatforms()
 		{
 			LinuxPlatformSDK SDK = new LinuxPlatformSDK();
-			SDK.ManageAndValidateSDK();
+			LinuxPlatformSDK SDKAarch64 = new LinuxPlatformSDK();
 
 			// Register this build platform for Linux x86-64 and AArch64
 			UEBuildPlatform.RegisterBuildPlatform(new LinuxPlatform(UnrealTargetPlatform.Linux, SDK));
@@ -911,7 +906,7 @@ namespace UnrealBuildTool
 			UEBuildPlatform.RegisterPlatformWithGroup(UnrealTargetPlatform.Linux, UnrealPlatformGroup.Unix);
 			UEBuildPlatform.RegisterPlatformWithGroup(UnrealTargetPlatform.Linux, UnrealPlatformGroup.Desktop);
 
-			UEBuildPlatform.RegisterBuildPlatform(new LinuxPlatform(UnrealTargetPlatform.LinuxAArch64, SDK));
+			UEBuildPlatform.RegisterBuildPlatform(new LinuxPlatform(UnrealTargetPlatform.LinuxAArch64, SDKAarch64));
 			UEBuildPlatform.RegisterPlatformWithGroup(UnrealTargetPlatform.LinuxAArch64, UnrealPlatformGroup.Linux);
 			UEBuildPlatform.RegisterPlatformWithGroup(UnrealTargetPlatform.LinuxAArch64, UnrealPlatformGroup.Unix);
 		}
diff --git a/Engine/Source/Programs/UnrealBuildTool/Platform/Lumin/UEBuildLumin.cs b/Engine/Source/Programs/UnrealBuildTool/Platform/Lumin/UEBuildLumin.cs
--- a/Engine/Source/Programs/UnrealBuildTool/Platform/Lumin/UEBuildLumin.cs
+++ b/Engine/Source/Programs/UnrealBuildTool/Platform/Lumin/UEBuildLumin.cs
@@ -8,6 +8,7 @@ using System.IO;
 using System.Xml;
 using System.Linq;
 using Tools.DotNETCommon;
+using System.Text.RegularExpressions;
 
 namespace UnrealBuildTool
 {
@@ -71,90 +72,9 @@ namespace UnrealBuildTool
 		#endregion
 	}
 
-	/// <summary>
-	/// Helper class for shared lumin platform functions and members
-	/// </summary>
-	public class LuminSDKVersionHelper
-	{
-		/// <summary>
-		/// This is the minimum SDK version we support.
-		/// </summary>
-		static uint MinimumSDKVersionMajor = 0;
-		static uint MinimumSDKVersionMinor = 23;
-
-		/// <summary>
-		/// Gets a string defining the minimum required sdk version.
-		/// </summary>
-		/// <returns></returns>
-		public string GetRequiredSDKString()
-		{
-			return string.Format("{0}.{1}", MinimumSDKVersionMajor, MinimumSDKVersionMinor);
-		}
-
-		/// <summary>
-		/// checks if the sdk is installed or has been synced, sets environment variable
-		/// </summary>
-		/// <returns></returns>
-		public bool HasAnySDK()
-		{
-			string EnvVarKey = "MLSDK";
-
-			string MLSDKPath = Environment.GetEnvironmentVariable(EnvVarKey);
-			if (String.IsNullOrEmpty(MLSDKPath))
-			{
-				Log.TraceLog("*** Unable to determine MLSDK location ***");
-				return false;
-			}
-
-			// detected SDK version is < minimum major/minor
-			uint DetectedMajorVersion;
-			uint DetectedMinorVersion;
-			String VersionFile = string.Format("{0}/include/ml_version.h", MLSDKPath).Replace('/', Path.DirectorySeparatorChar);
-			if (File.Exists(VersionFile))
-			{
-				string[] VersionText = File.ReadAllLines(VersionFile);
-
-				String MajorVersion = FindVersionNumber("MLSDK_VERSION_MAJOR", VersionText);
-				String MinorVersion = FindVersionNumber("MLSDK_VERSION_MINOR", VersionText);
-				DetectedMajorVersion = Convert.ToUInt32(MajorVersion);
-				DetectedMinorVersion = Convert.ToUInt32(MinorVersion);
-			}
-			else
-			{
-				Log.TraceLog("*** Unable to locate MLSDK version file ml_version.h ***");
-				return false;
-			}
-
-			if (DetectedMajorVersion < MinimumSDKVersionMajor || (DetectedMajorVersion == MinimumSDKVersionMajor && DetectedMinorVersion < MinimumSDKVersionMinor))
-			{
-				Log.TraceLog("*** Found installed MLSDK version {0}.{1} at '{2}' but require at least {3}.{4} ***",
-					DetectedMajorVersion, DetectedMinorVersion, MLSDKPath, MinimumSDKVersionMajor, MinimumSDKVersionMinor);
-				return false;
-			}
-
-			Log.TraceLog("*** Found installed MLSDK version {0}.{1} at '{2}' ***", DetectedMajorVersion, DetectedMinorVersion, MLSDKPath);
-			return true;
-		}
-
-		private string FindVersionNumber(string StringToFind, string[] AllLines)
-		{
-			string FoundVersion = "Unknown";
-			foreach (string CurrentLine in AllLines)
-			{
-				int Index = CurrentLine.IndexOf(StringToFind);
-				if (Index != -1)
-				{
-					FoundVersion = CurrentLine.Substring(Index + StringToFind.Length);
-					break;
-				}
-			}
-			return FoundVersion.Trim();
-		}
-	}
-
 	class LuminPlatform : AndroidPlatform
 	{
-		public LuminPlatform(AndroidPlatformSDK InSDK)
+		public LuminPlatform(UEBuildPlatformSDK InSDK)
 			: base(UnrealTargetPlatform.Lumin, InSDK)
 		{
 		}
@@ -324,51 +244,98 @@ namespace UnrealBuildTool
 		}
 	}
 
-	class LuminPlatformSDK : AndroidPlatformSDK
+	class LuminPlatformSDK : UEBuildPlatformSDK
 	{
-		/// <summary>
-		/// This is the minimum SDK version we support.
-		/// </summary>
-		LuminSDKVersionHelper SDKVersionHelper = new LuminSDKVersionHelper();
 
-		public override string GetSDKTargetPlatformName()
+		public override string GetDesiredVersion()
 		{
-			return "Lumin";
+			return "0.23";
 		}
 
-		protected override string GetRequiredSDKString()
-		{
-			return SDKVersionHelper.GetRequiredSDKString();
-		}
+		// not needed just returning GetDesiredVersion()
+// 		public override void GetValidVersionRange(out string MinVersion, out string MaxVersion)
+// 		{
+// 			MinVersion = "0.23";
+// 			MaxVersion = "0.23";
+// 		}
 
-		protected override String GetRequiredScriptVersionString()
-		{
-			return "Lumin_15";
-		}
 
-		/// <summary>
-		/// checks if the sdk is installed or has been synced, sets environment variable
-		/// </summary>
-		/// <returns></returns>
-		protected override bool HasAnySDK()
+		public override string GetInstalledSDKVersion()
 		{
-			return SDKVersionHelper.HasAnySDK();
+			string EnvVarKey = "MLSDK";
+
+			string MLSDKPath = Environment.GetEnvironmentVariable(EnvVarKey);
+			if (String.IsNullOrEmpty(MLSDKPath))
+			{
+				return null;
+			}
+
+			String VersionFile = Path.Combine(MLSDKPath, "include/ml_version.h");
+			if (File.Exists(VersionFile))
+			{
+				string[] VersionText = File.ReadAllLines(VersionFile);
+
+				String MajorVersion = FindVersionNumber("MLSDK_VERSION_MAJOR", VersionText);
+				String MinorVersion = FindVersionNumber("MLSDK_VERSION_MINOR", VersionText);
+
+				return string.Format("{0}.{1}", MajorVersion, MinorVersion);
+			}
+
+			return null;
 		}
 
-		protected override SDKStatus HasRequiredManualSDKInternal()
+
+		public override bool TryConvertVersionToInt(string StringValue, out UInt64 OutValue)
 		{
-			// if any autosdk setup has been done then the local process environment is suspect
-			if (HasSetupAutoSDK())
+			// make major high 16 bits, minor low 16
+			Match Result = Regex.Match(StringValue, @"^(\d+).(\d+)$");
+			if (Result.Success)
 			{
-				return SDKStatus.Invalid;
+				OutValue = UInt64.Parse(Result.Groups[1].Value) << 16 | UInt64.Parse(Result.Groups[2].Value) << 0;
+				return true;
 			}
+			OutValue = 0;
+			return false;
+		}
+
+
 
-			if (SDKVersionHelper.HasAnySDK())
+
+		private string FindVersionNumber(string StringToFind, string[] AllLines)
+		{
+			string FoundVersion = "Unknown";
+			foreach (string CurrentLine in AllLines)
 			{
-				return SDKStatus.Valid;
+				int Index = CurrentLine.IndexOf(StringToFind);
+				if (Index != -1)
+				{
+					FoundVersion = CurrentLine.Substring(Index + StringToFind.Length);
+					break;
+				}
 			}
+			return FoundVersion.Trim();
+		}
+
 
-			return SDKStatus.Invalid;
+		protected override bool PlatformSupportsAutoSDKs()
+		{
+			return true;
+		}
+
+		protected override string GetRequiredSDKString()
+		{
+			return GetDesiredVersion();
+		}
+
+		protected override String GetRequiredScriptVersionString()
+		{
+			return "Lumin_15";
+		}
+
+		// prefer auto sdk on android as correct 'manual' sdk detection isn't great at the moment.
+		protected override bool PreferAutoSDK()
+		{
+			return true;
 		}
 	}
 
@@ -382,7 +349,6 @@ namespace UnrealBuildTool
 		public override void RegisterBuildPlatforms()
 		{
 			LuminPlatformSDK SDK = new LuminPlatformSDK();
-			SDK.ManageAndValidateSDK();
 
 			// Register this build platform
 			UEBuildPlatform.RegisterBuildPlatform(new LuminPlatform(SDK));
diff --git a/Engine/Source/Programs/UnrealBuildTool/Platform/Mac/ApplePlatformSDK.cs b/Engine/Source/Programs/UnrealBuildTool/Platform/Mac/ApplePlatformSDK.cs
new file mode 100644
--- /dev/null
+++ b/Engine/Source/Programs/UnrealBuildTool/Platform/Mac/ApplePlatformSDK.cs
@@ -0,0 +1,121 @@
+// Copyright Epic Games, Inc. All Rights Reserved.
+
+using System;
+using System.Collections.Generic;
+using System.IO;
+using Tools.DotNETCommon;
+using System.Text.RegularExpressions;
+using Microsoft.Win32;
+using System.Diagnostics;
+
+namespace UnrealBuildTool
+{
+	internal class ApplePlatformSDK : UEBuildPlatformSDK
+	{
+		public override string GetDesiredVersion()
+		{
+			return "11.5";
+		}
+
+		public override void GetValidVersionRange(out string MinVersion, out string MaxVersion)
+		{
+			if (IsMac())
+			{
+				MinVersion = "9.4.0";
+				MaxVersion = "11.9.9";
+			}
+			else
+			{
+				// @todo turnkey: these are MobileDevice .dll versions in Windows - to get the iTunes app version (12.3.4.1 etc) would need to hunt down the .exe
+				MinVersion = "1100.0.0.0";
+				MaxVersion = "1199.0";
+			}
+		}
+
+
+		public override string GetInstalledSDKVersion()
+		{
+			// get xcode version on Mac
+			if (IsMac())
+			{
+				string Output = Utils.RunLocalProcessAndReturnStdOut("sh", "-c 'xcodebuild -version'");
+				Match Result = Regex.Match(Output, @"Xcode (\S*)");
+				return Result.Success ? Result.Groups[1].Value : "";
+			}
+
+			// otherwise, get iTunes "Version"
+			string DllPath = Registry.GetValue("HKEY_LOCAL_MACHINE\\SOFTWARE\\Wow6432Node\\Apple Inc.\\Apple Mobile Device Support\\Shared", "iTunesMobileDeviceDLL", null) as string;
+			if (string.IsNullOrEmpty(DllPath) || !File.Exists(DllPath))
+			{
+				DllPath = Registry.GetValue("HKEY_LOCAL_MACHINE\\SOFTWARE\\Wow6432Node\\Apple Inc.\\Apple Mobile Device Support\\Shared", "MobileDeviceDLL", null) as string;
+				if (string.IsNullOrEmpty(DllPath) || !File.Exists(DllPath))
+				{
+					DllPath = FindWindowsStoreITunesDLL();
+
+				}
+			}
+
+			if (!string.IsNullOrEmpty(DllPath) && File.Exists(DllPath))
+			{
+				return FileVersionInfo.GetVersionInfo(DllPath).FileVersion;
+			}
+
+			return null;
+		}
+
+		public override bool TryConvertVersionToInt(string StringValue, out UInt64 OutValue)
+		{
+			// 8 bits per component, with high getting extra from high 32
+			Match Result = Regex.Match(StringValue, @"^(\d+).(\d+)(.(\d+))?(.(\d+))?$");
+			if (Result.Success)
+			{
+				OutValue = UInt64.Parse(Result.Groups[1].Value) << 24 | UInt64.Parse(Result.Groups[2].Value) << 16;
+				if (Result.Groups[4].Success)
+				{
+					OutValue = UInt64.Parse(Result.Groups[4].Value) << 8;
+				}
+				if (Result.Groups[6].Success)
+				{
+					OutValue = UInt64.Parse(Result.Groups[6].Value) << 0;
+				}
+				return true;
+			}
+			OutValue = 0;
+			return false;
+		}
+
+
+
+
+
+		static string FindWindowsStoreITunesDLL()
+		{
+			string InstallPath = null;
+
+			string PackagesKeyName = "Software\\Classes\\Local Settings\\Software\\Microsoft\\Windows\\CurrentVersion\\AppModel\\PackageRepository\\Packages";
+
+			RegistryKey PackagesKey = Registry.LocalMachine.OpenSubKey(PackagesKeyName);
+			if (PackagesKey != null)
+			{
+				string[] PackageSubKeyNames = PackagesKey.GetSubKeyNames();
+
+				foreach (string PackageSubKeyName in PackageSubKeyNames)
+				{
+					if (PackageSubKeyName.Contains("AppleInc.iTunes") && (PackageSubKeyName.Contains("_x64") || PackageSubKeyName.Contains("_x86")))
+					{
+						string FullPackageSubKeyName = PackagesKeyName + "\\" + PackageSubKeyName;
+
+						RegistryKey iTunesKey = Registry.LocalMachine.OpenSubKey(FullPackageSubKeyName);
+						if (iTunesKey != null)
+						{
+							InstallPath = (string)iTunesKey.GetValue("Path") + "\\AMDS32\\MobileDevice.dll";
+							break;
+						}
+					}
+				}
+			}
+
+			return InstallPath;
+		}
+	}
+}
diff --git a/Engine/Source/Programs/UnrealBuildTool/Platform/Mac/UEBuildMac.cs b/Engine/Source/Programs/UnrealBuildTool/Platform/Mac/UEBuildMac.cs
--- a/Engine/Source/Programs/UnrealBuildTool/Platform/Mac/UEBuildMac.cs
+++ b/Engine/Source/Programs/UnrealBuildTool/Platform/Mac/UEBuildMac.cs
@@ -3,9 +3,6 @@
 
 using System;
 using System.Collections.Generic;
-using System.Text;
-using System.Diagnostics;
-using System.IO;
 using Tools.DotNETCommon;
 
 namespace UnrealBuildTool
@@ -99,16 +96,8 @@ namespace UnrealBuildTool
 
 	class MacPlatform : UEBuildPlatform
 	{
-		MacPlatformSDK SDK;
-
-		public MacPlatform(MacPlatformSDK InSDK) : base(UnrealTargetPlatform.Mac)
+		public MacPlatform(UEBuildPlatformSDK InSDK) : base(UnrealTargetPlatform.Mac, InSDK)
 		{
-			SDK = InSDK;
-		}
-
-		public override SDKStatus HasRequiredSDKsInstalled()
-		{
-			return SDK.HasRequiredSDKsInstalled();
 		}
 
 		public override bool CanUseXGE()
@@ -363,14 +352,6 @@ namespace UnrealBuildTool
 		}
 	}
 
-	class MacPlatformSDK : UEBuildPlatformSDK
-	{
-		protected override SDKStatus HasRequiredManualSDKInternal()
-		{
-			return SDKStatus.Valid;
-		}
-	}
-
 	class MacPlatformFactory : UEBuildPlatformFactory
 	{
 		public override UnrealTargetPlatform TargetPlatform
@@ -383,8 +364,7 @@ namespace UnrealBuildTool
 		/// </summary>
 		public override void RegisterBuildPlatforms()
 		{
-			MacPlatformSDK SDK = new MacPlatformSDK();
-			SDK.ManageAndValidateSDK();
+			ApplePlatformSDK SDK = new ApplePlatformSDK();
 
 			// Register this build platform for Mac
 			UEBuildPlatform.RegisterBuildPlatform(new MacPlatform(SDK));
diff --git a/Engine/Source/Programs/UnrealBuildTool/Platform/TVOS/UEBuildTVOS.cs b/Engine/Source/Programs/UnrealBuildTool/Platform/TVOS/UEBuildTVOS.cs
--- a/Engine/Source/Programs/UnrealBuildTool/Platform/TVOS/UEBuildTVOS.cs
+++ b/Engine/Source/Programs/UnrealBuildTool/Platform/TVOS/UEBuildTVOS.cs
@@ -2,10 +2,6 @@
 
 using System;
 using System.Collections.Generic;
-using System.Text;
-using System.Diagnostics;
-using System.IO;
-using System.Xml;
 using Tools.DotNETCommon;
 
 namespace UnrealBuildTool
@@ -44,7 +40,7 @@ namespace UnrealBuildTool
 
 	class TVOSPlatform : IOSPlatform
     {
-		public TVOSPlatform(IOSPlatformSDK InSDK)
+		public TVOSPlatform(UEBuildPlatformSDK InSDK)
 			: base(InSDK, UnrealTargetPlatform.TVOS)
 		{
 		}
@@ -157,9 +153,8 @@ namespace UnrealBuildTool
 		/// </summary>
 		public override void RegisterBuildPlatforms()
 		{
-			IOSPlatformSDK SDK = new IOSPlatformSDK();
-			SDK.ManageAndValidateSDK();
-
+			ApplePlatformSDK SDK = new ApplePlatformSDK();
+		
 			// Register this build platform for IOS
 			UEBuildPlatform.RegisterBuildPlatform(new TVOSPlatform(SDK));
 			UEBuildPlatform.RegisterPlatformWithGroup(UnrealTargetPlatform.TVOS, UnrealPlatformGroup.Apple);
diff --git a/Engine/Source/Programs/UnrealBuildTool/Platform/Windows/MicrosoftPlatformSDK.cs b/Engine/Source/Programs/UnrealBuildTool/Platform/Windows/MicrosoftPlatformSDK.cs
new file mode 100644
--- /dev/null
+++ b/Engine/Source/Programs/UnrealBuildTool/Platform/Windows/MicrosoftPlatformSDK.cs
@@ -0,0 +1,63 @@
+// Copyright Epic Games, Inc. All Rights Reserved.
+
+using System;
+using System.Collections.Generic;
+using Tools.DotNETCommon;
+using System.Text.RegularExpressions;
+using Microsoft.Win32;
+
+namespace UnrealBuildTool
+{
+	internal class MicrosoftPlatformSDK : UEBuildPlatformSDK
+	{
+		public override string GetDesiredVersion()
+		{
+			// the current and previous versions of the SDK (technically, NX Add-On), with appended NEX version, we support
+			return "10.0.18632";
+		}
+
+		public override void GetValidVersionRange(out string MinVersion, out string MaxVersion)
+		{
+			MinVersion = "10.0.00000";
+			MaxVersion = "10.9.99999";
+		}
+
+		public override string GetInstalledSDKVersion()
+		{
+			if (Utils.IsRunningOnMono)
+			{
+				return null;
+			}
+
+			// @todo turnkey: do we support pre-10?
+			// @todo turnkey: MicrosoftPlatformSDK maybe?
+			string Version = "v10.0";
+			object Result = Registry.GetValue(@"HKEY_LOCAL_MACHINE\SOFTWARE\Wow6432Node\Microsoft\Microsoft SDKs\Windows\" + Version, "InstallationFolder", null);
+
+			if (Result != null)
+			{
+				return Registry.GetValue(@"HKEY_LOCAL_MACHINE\SOFTWARE\Wow6432Node\Microsoft\Microsoft SDKs\Windows\" + Version, "ProductVersion", null) as string;
+			}
+
+			return null;
+		}
+
+		public override bool TryConvertVersionToInt(string StringValue, out UInt64 OutValue)
+		{
+			OutValue = 0;
+
+			Match Result = Regex.Match(StringValue, @"^(\d+).(\d+).(\d+)");
+			if (Result.Success)
+			{
+				// 8 bits for major, 8 for minor, 16 for patch
+				OutValue |= UInt64.Parse(Result.Groups[1].Value) << 24;
+				OutValue |= UInt64.Parse(Result.Groups[2].Value) << 16;
+				OutValue |= UInt64.Parse(Result.Groups[3].Value) << 0;
+
+				return true;
+			}
+
+			return false;
+		}
+	}
+}
diff --git a/Engine/Source/Programs/UnrealBuildTool/Platform/Windows/UEBuildWindows.cs b/Engine/Source/Programs/UnrealBuildTool/Platform/Windows/UEBuildWindows.cs
--- a/Engine/Source/Programs/UnrealBuildTool/Platform/Windows/UEBuildWindows.cs
+++ b/Engine/Source/Programs/UnrealBuildTool/Platform/Windows/UEBuildWindows.cs
@@ -738,19 +738,11 @@ namespace UnrealBuildTool
 		/// <param name="InPlatform">Creates a windows platform with the given enum value</param>
 		/// <param name="InSDK">The installed Windows SDK</param>
 		public WindowsPlatform(UnrealTargetPlatform InPlatform, WindowsPlatformSDK InSDK)
-			: base(InPlatform)
+			: base(InPlatform, InSDK)
 		{
 			SDK = InSDK;
 		}
 
-		/// <summary>
-		/// Whether the required external SDKs are installed for this platform. Could be either a manual install or an AutoSDK.
-		/// </summary>
-		public override SDKStatus HasRequiredSDKsInstalled()
-		{
-			return SDK.HasRequiredSDKsInstalled();
-		}
-
 		/// <summary>
 		/// Reset a target's settings to the default
 		/// </summary>
@@ -2307,12 +2299,8 @@ namespace UnrealBuildTool
 		}
 	}
 
-	class WindowsPlatformSDK : UEBuildPlatformSDK
+	class WindowsPlatformSDK : MicrosoftPlatformSDK
 	{
-		protected override SDKStatus HasRequiredManualSDKInternal()
-		{
-			return SDKStatus.Valid;
-		}
 	}
 
 	class WindowsPlatformFactory : UEBuildPlatformFactory
@@ -2328,7 +2316,6 @@ namespace UnrealBuildTool
 		public override void RegisterBuildPlatforms()
 		{
 			WindowsPlatformSDK SDK = new WindowsPlatformSDK();
-			SDK.ManageAndValidateSDK();
 
 			// Register this build platform for Win64 (no more Win32)
 			UEBuildPlatform.RegisterBuildPlatform(new WindowsPlatform(UnrealTargetPlatform.Win64, SDK));
diff --git a/Engine/Source/Programs/UnrealBuildTool/System/DynamicCompilation.cs b/Engine/Source/Programs/UnrealBuildTool/System/DynamicCompilation.cs
--- a/Engine/Source/Programs/UnrealBuildTool/System/DynamicCompilation.cs
+++ b/Engine/Source/Programs/UnrealBuildTool/System/DynamicCompilation.cs
@@ -308,6 +308,9 @@ namespace UnrealBuildTool
 					// by looking for the one that contains a common utility class
 					Assembly UtilitiesAssembly = Assembly.GetAssembly(typeof(FileReference));
 					CompileParams.ReferencedAssemblies.Add(UtilitiesAssembly.Location);
+
+					Assembly BuildUtilitiesAssembly = Assembly.GetAssembly(typeof(UEBuildPlatformSDK));
+					CompileParams.ReferencedAssemblies.Add(BuildUtilitiesAssembly.Location);
 				}
 
 				// Add preprocessor definitions
diff --git a/Engine/Source/Programs/UnrealBuildTool/System/XmlConfig.cs b/Engine/Source/Programs/UnrealBuildTool/System/XmlConfig.cs
--- a/Engine/Source/Programs/UnrealBuildTool/System/XmlConfig.cs
+++ b/Engine/Source/Programs/UnrealBuildTool/System/XmlConfig.cs
@@ -154,13 +154,21 @@ namespace UnrealBuildTool
 		static List<Type> FindConfigurableTypes()
 		{
 			List<Type> ConfigTypes = new List<Type>();
-			foreach(Type ConfigType in Assembly.GetExecutingAssembly().GetTypes())
+			try
 			{
-				if(HasXmlConfigFileAttribute(ConfigType))
+				foreach (Type ConfigType in Assembly.GetExecutingAssembly().GetTypes())
 				{
-					ConfigTypes.Add(ConfigType);
+					if (HasXmlConfigFileAttribute(ConfigType))
+					{
+						ConfigTypes.Add(ConfigType);
+					}
 				}
 			}
+			catch (ReflectionTypeLoadException Ex)
+			{
+				Console.WriteLine("TypeLoadException: {0}", string.Join("\n", Ex.LoaderExceptions.Select(x => x.Message)));
+				throw;
+			}
 			return ConfigTypes;
 		}
 
diff --git a/Engine/Source/Programs/UnrealBuildTool/UnrealBuildTool.csproj b/Engine/Source/Programs/UnrealBuildTool/UnrealBuildTool.csproj
--- a/Engine/Source/Programs/UnrealBuildTool/UnrealBuildTool.csproj
+++ b/Engine/Source/Programs/UnrealBuildTool/UnrealBuildTool.csproj
@@ -158,7 +158,6 @@
     <Compile Include="Configuration\TargetDescriptor.cs" />
     <Compile Include="Configuration\UEBuildModuleExternal.cs" />
     <Compile Include="Configuration\UEBuildPlatformFactory.cs" />
-    <Compile Include="Configuration\UEBuildPlatformSDK.cs" />
     <Compile Include="Configuration\UEBuildPlugin.cs" />
     <Compile Include="Executors\Distcc.cs" />
     <Compile Include="Executors\LocalExecutor.cs" />
@@ -170,6 +169,8 @@
     <Compile Include="Platform\IOS\MobileProvisionContents.cs" />
     <Compile Include="Platform\Lumin\LuminPackageManifest.cs" />
     <Compile Include="Platform\Lumin\UEDeployLumin.cs" />
+    <Compile Include="Platform\Mac\ApplePlatformSDK.cs" />
+    <Compile Include="Platform\Windows\MicrosoftPlatformSDK.cs" />
     <Compile Include="Platform\Windows\VCManifestGenerator.cs" />
     <Compile Include="ProjectFiles\CLion\CLionGenerator.cs" />
     <Compile Include="ProjectFiles\CMake\CMakefileGenerator.cs" />
@@ -451,6 +452,10 @@
     </PreBuildEvent>
   </PropertyGroup>
   <ItemGroup>
+    <ProjectReference Include="..\DotNETCommon\BuildUtilities\BuildUtilities.csproj">
+      <Project>{5965bed9-4823-4671-84a5-c8d6503d5848}</Project>
+      <Name>BuildUtilities</Name>
+    </ProjectReference>
     <ProjectReference Include="..\DotNETCommon\DotNETUtilities\DotNETUtilities.csproj">
       <Project>{5d7d66e8-8c76-4af9-b3ec-2ef03421d730}</Project>
       <Name>DotNETUtilities</Name>
