Rename "Scopes" to "Spans" in exported telemetry data.

diff --git a/Engine/Source/Programs/DotNETCommon/DotNETUtilities/TraceSpan.cs b/Engine/Source/Programs/DotNETCommon/DotNETUtilities/TraceSpan.cs
--- a/Engine/Source/Programs/DotNETCommon/DotNETUtilities/TraceSpan.cs
+++ b/Engine/Source/Programs/DotNETCommon/DotNETUtilities/TraceSpan.cs
@@ -94,17 +94,17 @@ namespace Tools.DotNETCommon
 				using (JsonWriter Writer = new JsonWriter(File))
 				{
 					Writer.WriteObjectStart();
-					Writer.WriteArrayStart("Scopes");
-					foreach (TraceSpanImpl Scope in Spans)
+					Writer.WriteArrayStart("Spans");
+					foreach (TraceSpanImpl Span in Spans)
 					{
-						if (Scope.FinishTime != null)
+						if (Span.FinishTime != null)
 						{
 							Writer.WriteObjectStart();
-							Writer.WriteValue("Name", Scope.Name);
-							Writer.WriteValue("StartTime", Scope.StartTime.ToString("o", CultureInfo.InvariantCulture));
-							Writer.WriteValue("FinishTime", Scope.FinishTime.Value.ToString("o", CultureInfo.InvariantCulture));
+							Writer.WriteValue("Name", Span.Name);
+							Writer.WriteValue("StartTime", Span.StartTime.ToString("o", CultureInfo.InvariantCulture));
+							Writer.WriteValue("FinishTime", Span.FinishTime.Value.ToString("o", CultureInfo.InvariantCulture));
 							Writer.WriteObjectStart("Metadata");
-							foreach (KeyValuePair<string, string> Pair in Scope.Metadata)
+							foreach (KeyValuePair<string, string> Pair in Span.Metadata)
 							{
 								Writer.WriteValue(Pair.Key, Pair.Value);
 							}
