Add missing copyright notice.

diff --git a/Engine/Source/Programs/DotNETCommon/DotNETUtilities/TraceSpan.cs b/Engine/Source/Programs/DotNETCommon/DotNETUtilities/TraceSpan.cs
--- a/Engine/Source/Programs/DotNETCommon/DotNETUtilities/TraceSpan.cs
+++ b/Engine/Source/Programs/DotNETCommon/DotNETUtilities/TraceSpan.cs
@@ -1,3 +1,5 @@
+// Copyright Epic Games, Inc. All Rights Reserved.
+
 using System;
 using System.Collections.Generic;
 using System.Diagnostics;
@@ -87,8 +89,8 @@ namespace Tools.DotNETCommon
 					DirectoryReference TelemetryDirRef = new DirectoryReference(TelemetryDir);
 					DirectoryReference.CreateDirectory(TelemetryDirRef);
 
-					string Name = String.Format("{0}.{1}", Path.GetFileName(Assembly.GetEntryAssembly().Location), Process.Id, Process.StartTime.Ticks);
-					File = FileReference.Combine(TelemetryDirRef, Name);
+					string FileName = String.Format("{0}.{1}.json", Path.GetFileName(Assembly.GetEntryAssembly().Location), Process.Id, Process.StartTime.Ticks);
+					File = FileReference.Combine(TelemetryDirRef, FileName);
 				}
 
 				using (JsonWriter Writer = new JsonWriter(File))
