Add additional tracing into to UAT/UBT.

diff --git a/Engine/Source/Programs/AutomationTool/BuildGraph/Task.cs b/Engine/Source/Programs/AutomationTool/BuildGraph/Task.cs
--- a/Engine/Source/Programs/AutomationTool/BuildGraph/Task.cs
+++ b/Engine/Source/Programs/AutomationTool/BuildGraph/Task.cs
@@ -210,8 +210,8 @@ namespace AutomationTool
 		/// <param name="Prefix">Prefix for metadata entries</param>
 		public virtual void GetTraceMetadata(ITraceSpan Span, string Prefix)
 		{
-			Span.AddMetadata(Prefix + "file", SourceLocation.Item1.MakeRelativeTo(CommandUtils.RootDirectory));
-			Span.AddMetadata(Prefix + "line", SourceLocation.Item2.ToString());
+			Span.AddMetadata(Prefix + "source.file", SourceLocation.Item1.MakeRelativeTo(CommandUtils.RootDirectory));
+			Span.AddMetadata(Prefix + "source.line", SourceLocation.Item2.ToString());
 		}
 
 		/// <summary>
diff --git a/Engine/Source/Programs/AutomationTool/BuildGraph/Tasks/CompileTask.cs b/Engine/Source/Programs/AutomationTool/BuildGraph/Tasks/CompileTask.cs
--- a/Engine/Source/Programs/AutomationTool/BuildGraph/Tasks/CompileTask.cs
+++ b/Engine/Source/Programs/AutomationTool/BuildGraph/Tasks/CompileTask.cs
@@ -237,9 +237,14 @@ namespace AutomationTool
 		{
 			base.GetTraceMetadata(Span, Prefix);
 
-			Span.AddMetadata(Prefix + "target", Parameters.Target);
-			Span.AddMetadata(Prefix + "config", Parameters.Configuration.ToString());
-			Span.AddMetadata(Prefix + "platform", Parameters.Platform.ToString());
+			Span.AddMetadata(Prefix + "target.name", Parameters.Target);
+			Span.AddMetadata(Prefix + "target.config", Parameters.Configuration.ToString());
+			Span.AddMetadata(Prefix + "target.platform", Parameters.Platform.ToString());
+
+			if (Parameters.Project != null)
+			{
+				Span.AddMetadata(Prefix + "target.project", Parameters.Project);
+			}
 		}
 
 		/// <summary>
diff --git a/Engine/Source/Programs/DotNETCommon/DotNETUtilities/TraceSpan.cs b/Engine/Source/Programs/DotNETCommon/DotNETUtilities/TraceSpan.cs
--- a/Engine/Source/Programs/DotNETCommon/DotNETUtilities/TraceSpan.cs
+++ b/Engine/Source/Programs/DotNETCommon/DotNETUtilities/TraceSpan.cs
@@ -36,13 +36,17 @@ namespace Tools.DotNETCommon
 		class TraceSpanImpl : ITraceSpan
 		{
 			public string Name;
+			public string Resource;
+			public string Service;
 			public DateTimeOffset StartTime;
 			public DateTimeOffset? FinishTime;
 			public Dictionary<string, string> Metadata = new Dictionary<string, string>();
 
-			public TraceSpanImpl(string Name)
+			public TraceSpanImpl(string Name, string Resource, string Service)
 			{
 				this.Name = Name;
+				this.Resource = Resource;
+				this.Service = Service;
 				this.StartTime = DateTimeOffset.Now;
 			}
 
@@ -68,9 +72,9 @@ namespace Tools.DotNETCommon
 		/// <summary>
 		/// Creates a scope using the current provider
 		/// </summary>
-		public static ITraceSpan Create(string Name)
+		public static ITraceSpan Create(string Name, string Resource = null, string Service = null)
 		{
-			TraceSpanImpl Span = new TraceSpanImpl(Name);
+			TraceSpanImpl Span = new TraceSpanImpl(Name, Resource, Service);
 			Spans.Add(Span);
 			return Span;
 		}
@@ -103,6 +107,14 @@ namespace Tools.DotNETCommon
 						{
 							Writer.WriteObjectStart();
 							Writer.WriteValue("Name", Span.Name);
+							if (Span.Resource != null)
+							{
+								Writer.WriteValue("Resource", Span.Resource);
+							}
+							if (Span.Service != null)
+							{
+								Writer.WriteValue("Service", Span.Service);
+							}
 							Writer.WriteValue("StartTime", Span.StartTime.ToString("o", CultureInfo.InvariantCulture));
 							Writer.WriteValue("FinishTime", Span.FinishTime.Value.ToString("o", CultureInfo.InvariantCulture));
 							Writer.WriteObjectStart("Metadata");
diff --git a/Engine/Source/Programs/UnrealBuildTool/System/Timeline.cs b/Engine/Source/Programs/UnrealBuildTool/System/Timeline.cs
--- a/Engine/Source/Programs/UnrealBuildTool/System/Timeline.cs
+++ b/Engine/Source/Programs/UnrealBuildTool/System/Timeline.cs
@@ -46,6 +46,11 @@ namespace UnrealBuildTool
 			/// </summary>
 			public TimeSpan? FinishTime;
 
+			/// <summary>
+			/// The trace span for external tracing
+			/// </summary>
+			public ITraceSpan Span;
+
 			/// <summary>
 			/// Constructor
 			/// </summary>
@@ -57,6 +62,7 @@ namespace UnrealBuildTool
 				this.Name = Name;
 				this.StartTime = StartTime;
 				this.FinishTime = FinishTime;
+				this.Span = TraceSpan.Create(Name);
 			}
 
 			/// <summary>
@@ -68,6 +74,7 @@ namespace UnrealBuildTool
 				{
 					FinishTime = Stopwatch.Elapsed;
 				}
+				Span.Dispose();
 			}
 
 			/// <summary>
diff --git a/Engine/Source/Programs/UnrealBuildTool/UnrealBuildTool.cs b/Engine/Source/Programs/UnrealBuildTool/UnrealBuildTool.cs
--- a/Engine/Source/Programs/UnrealBuildTool/UnrealBuildTool.cs
+++ b/Engine/Source/Programs/UnrealBuildTool/UnrealBuildTool.cs
@@ -600,8 +600,11 @@ namespace UnrealBuildTool
 				// Make sure we flush the logs however we exit
 				Trace.Close();
 
+				// Write any trace logs
+				TraceSpan.Flush();
+
 				// Dispose of the mutex. Must be done last to ensure that another process does not startup and start trying to write to the same log file.
-				if(Mutex != null)
+				if (Mutex != null)
 				{
 					Mutex.Dispose();
 				}
