Add support for passing properties to the <CsCompile> task, and fix parsing of ShooterGame project files in the BuildEditorAndTools.xml script.

#rnx

diff --git a/Engine/Build/Graph/Examples/BuildEditorAndTools.xml b/Engine/Build/Graph/Examples/BuildEditorAndTools.xml
--- a/Engine/Build/Graph/Examples/BuildEditorAndTools.xml
+++ b/Engine/Build/Graph/Examples/BuildEditorAndTools.xml
@@ -106,7 +106,7 @@
 			<Tag Files="Engine/Source/Programs/AutomationTool/..." Filter="*.csproj" With="#UAT Projects"/>
 			<Tag Files="Engine/Platforms/..." Filter="*.Automation.csproj" With="#UAT Projects"/>
 			<Tag Files="$(RootDir)/.../Build/Scripts/..." Filter="*.Automation.csproj" With="#UAT Projects"/>
-			<CsCompile Project="#UAT Projects" Configuration="Development" Platform="AnyCPU" Tag="#ArchiveFiles" EnumerateOnly="true"/>
+			<CsCompile Project="#UAT Projects" Configuration="Development" Platform="AnyCPU" Properties="EngineDir=$(RootDir)/Engine" Tag="#ArchiveFiles" EnumerateOnly="true"/>
 			
 			<!-- Write a text file marker indicating that we're using precompiled binaries -->
 			<Spawn Exe="cmd.exe" Arguments="/C echo. &gt;&quot;$(RootDir)/Engine/Build/PrecompiledBinaries.txt&quot;"/>
diff --git a/Engine/Source/Programs/AutomationTool/BuildGraph/Tasks/CsCompileTask.cs b/Engine/Source/Programs/AutomationTool/BuildGraph/Tasks/CsCompileTask.cs
--- a/Engine/Source/Programs/AutomationTool/BuildGraph/Tasks/CsCompileTask.cs
+++ b/Engine/Source/Programs/AutomationTool/BuildGraph/Tasks/CsCompileTask.cs
@@ -40,6 +40,12 @@ namespace AutomationTool.Tasks
 		[TaskParameter(Optional = true)]
 		public string Target;
 
+		/// <summary>
+		/// Properties for the command
+		/// </summary>
+		[TaskParameter(Optional = true)]
+		public string Properties;
+
 		/// <summary>
 		/// Additional options to pass to the compiler.
 		/// </summary>
@@ -118,6 +124,24 @@ namespace AutomationTool.Tasks
 			{
 				Properties["Configuration"] = Parameters.Configuration;
 			}
+			if (!String.IsNullOrEmpty(Parameters.Properties))
+			{
+				foreach (string Property in Parameters.Properties.Split(';'))
+				{
+					if (!String.IsNullOrWhiteSpace(Property))
+					{
+						int EqualsIdx = Property.IndexOf('=');
+						if (EqualsIdx == -1)
+						{
+							Log.TraceWarning("Missing '=' in property assignment");
+						}
+						else
+						{
+							Properties[Property.Substring(0, EqualsIdx).Trim()] = Property.Substring(EqualsIdx + 1).Trim();
+						}
+					}
+				}
+			}
 
 			// Build the arguments and run the build
 			if(!Parameters.EnumerateOnly)
diff --git a/Engine/Source/Programs/DotNETCommon/DotNETUtilities/CsProjectInfo.cs b/Engine/Source/Programs/DotNETCommon/DotNETUtilities/CsProjectInfo.cs
--- a/Engine/Source/Programs/DotNETCommon/DotNETUtilities/CsProjectInfo.cs
+++ b/Engine/Source/Programs/DotNETCommon/DotNETUtilities/CsProjectInfo.cs
@@ -402,7 +402,7 @@ namespace Tools.DotNETCommon
 						// Reference to another project
 						if (EvaluateCondition(ItemElement, ProjectInfo))
 						{
-							ParseProjectReference(BaseDirectory, ItemElement, ProjectInfo.ProjectReferences);
+							ParseProjectReference(BaseDirectory, ItemElement, ProjectInfo.Properties, ProjectInfo.ProjectReferences);
 						}
 						break;
 					case "Compile":
@@ -451,13 +451,14 @@ namespace Tools.DotNETCommon
 		/// </summary>
 		/// <param name="BaseDirectory">Directory to resolve relative paths against</param>
 		/// <param name="ParentElement">The parent 'ProjectReference' element</param>
+		/// <param name="Properties">Dictionary of properties for parsing the file</param>
 		/// <param name="ProjectReferences">Dictionary of project files to a bool indicating whether the outputs of the project should be copied locally to the referencing project.</param>
-		static void ParseProjectReference(DirectoryReference BaseDirectory, XmlElement ParentElement, Dictionary<FileReference, bool> ProjectReferences)
+		static void ParseProjectReference(DirectoryReference BaseDirectory, XmlElement ParentElement, Dictionary<string, string> Properties, Dictionary<FileReference, bool> ProjectReferences)
 		{
 			string IncludePath = UnescapeString(ParentElement.GetAttribute("Include"));
 			if (!String.IsNullOrEmpty(IncludePath))
 			{
-				FileReference ProjectFile = FileReference.Combine(BaseDirectory, IncludePath);
+				FileReference ProjectFile = FileReference.Combine(BaseDirectory, ExpandProperties(IncludePath, Properties));
 				bool bPrivate = GetChildElementBoolean(ParentElement, "Private", true);
 				ProjectReferences[ProjectFile] = bPrivate;
 			}
