- Use case insensitive comparisons int he AssemblyLocationCache as some assemblies names don't match the case of their filenames (the cache is generated by parsing the filename, not loading the assembly to get internal name)

diff --git a/Engine/Source/Programs/DotNETCommon/DotNETUtilities/AssemblyUtils.cs b/Engine/Source/Programs/DotNETCommon/DotNETUtilities/AssemblyUtils.cs
--- a/Engine/Source/Programs/DotNETCommon/DotNETUtilities/AssemblyUtils.cs
+++ b/Engine/Source/Programs/DotNETCommon/DotNETUtilities/AssemblyUtils.cs
@@ -121,9 +121,9 @@ namespace Tools.DotNETCommon
 		}
 
 		// Map of assembly name to path on disk
-		private static Dictionary<string, string> AssemblyLocationCache = new Dictionary<string, string>();
+		private static Dictionary<string, string> AssemblyLocationCache = new Dictionary<string, string>(StringComparer.OrdinalIgnoreCase);
 		// Track last modified date of each assembly, so we can ensure we always reference the latest one in the case of stale assemblies on disk.
-		private static Dictionary<string, DateTime> AssemblyWriteTimes = new Dictionary<string, DateTime>();
+		private static Dictionary<string, DateTime> AssemblyWriteTimes = new Dictionary<string, DateTime>(StringComparer.OrdinalIgnoreCase);
 
     }
 }
