Added UAT commands with helpers in DotNetUtilities for uploading and download trees of content to Jupiter (our new storage service).

diff --git a/Engine/Source/Programs/AutomationTool/BuildGraph/BuildGraph.Automation.csproj b/Engine/Source/Programs/AutomationTool/BuildGraph/BuildGraph.Automation.csproj
--- a/Engine/Source/Programs/AutomationTool/BuildGraph/BuildGraph.Automation.csproj
+++ b/Engine/Source/Programs/AutomationTool/BuildGraph/BuildGraph.Automation.csproj
@@ -51,6 +51,7 @@
     <Reference Include="System.Core" />
     <Reference Include="System.IO.Compression" />
     <Reference Include="System.IO.Compression.FileSystem" />
+    <Reference Include="System.Net.Http" />
     <Reference Include="System.Runtime.Serialization" />
     <Reference Include="System.Web" />
     <Reference Include="System.Web.Extensions" />
@@ -103,6 +104,7 @@
     <Compile Include="Tasks\SetVersionTask.cs" />
     <Compile Include="Tasks\UnzipTask.cs" />
     <Compile Include="Tasks\WriteTextFileTask.cs" />
+    <Compile Include="Tasks\UploadTask.cs" />
     <Compile Include="Tasks\ZipTask.cs" />
     <Compile Include="TempStorage.cs" />
     <Compile Include="Tasks\ModifyConfigTask.cs" />
diff --git a/Engine/Source/Programs/AutomationTool/BuildGraph/Tasks/UploadTask.cs b/Engine/Source/Programs/AutomationTool/BuildGraph/Tasks/UploadTask.cs
new file mode 100644
--- /dev/null
+++ b/Engine/Source/Programs/AutomationTool/BuildGraph/Tasks/UploadTask.cs
@@ -0,0 +1,193 @@
+// Copyright Epic Games, Inc. All Rights Reserved.
+
+using AutomationTool;
+using System;
+using System.Collections.Generic;
+using System.IO;
+using System.Linq;
+using System.Net.Http;
+using System.Security.Cryptography;
+using System.Text;
+using System.Threading.Tasks;
+using System.Web.UI.HtmlControls;
+using System.Xml;
+using Tools.DotNETCommon;
+using Tools.DotNETCommon.Jupiter;
+
+namespace BuildGraph.Tasks
+{
+	/// <summary>
+	/// Parameters for a zip task
+	/// </summary>
+	public class UploadTaskParameters
+	{
+		/// <summary>
+		/// The directory to read compressed files from.
+		/// </summary>
+		[TaskParameter]
+		public DirectoryReference FromDir;
+
+		/// <summary>
+		/// List of file specifications separated by semicolons (for example, *.cpp;Engine/.../*.bat), or the name of a tag set. Relative paths are taken from FromDir.
+		/// </summary>
+		[TaskParameter(Optional = true, ValidationType = TaskParameterValidationType.FileSpec)]
+		public string Files;
+		
+		/// <summary>
+		/// The jupiter namespace used to upload the build. Used to control who has access to the build.
+		/// </summary>
+		[TaskParameter]
+		public string JupiterNamespace;
+
+		/// <summary>
+		/// The key of the build as will be used to download the build again. This has to be globally unique for this particular upload.
+		/// </summary>
+		[TaskParameter]
+		public string JupiterKey;
+
+		/// <summary>
+		/// The type of archive these files are from, will be added to the metadata
+		/// </summary>
+		[TaskParameter]
+		public string ArchiveType;
+
+		/// <summary>
+		/// The name of the project this set of files are associated with, will be added to the metadata
+		/// </summary>
+		[TaskParameter]
+		public string ProjectName;
+
+		/// <summary>
+		/// The source control branch these files were generated from, will be added to the metadata
+		/// </summary>
+		[TaskParameter]
+		public string Branch;
+
+		/// <summary>
+		/// The source control revision these files were generated from, will be added to the metadata
+		/// </summary>
+		[TaskParameter]
+		public string Changelist;
+
+		/// <summary>
+		/// Specify the url to the Jupiter instance to upload to
+		/// </summary>
+		[TaskParameter]
+		public string JupiterUrl;
+
+		/// <summary>
+		/// Semi-colon separated list of '=' separated key value mappings to add to the metadata. E.g. Foo=bar;spam=eggs
+		/// </summary>
+		[TaskParameter(Optional = true)]
+		public string AdditionalMetadata;
+
+		/// <summary>
+		/// If enabled file content is not kept in memory, results in lower memory usage but increased io as file contents needs to be read multiple times (for hashing as well as during upload)
+		/// </summary>
+		[TaskParameter(Optional = true)]
+		public bool LimitMemoryUsage = true;
+	}
+
+	/// <summary>
+	/// Uploads a set of files to Jupiter for future retrival
+	/// </summary>
+	[TaskElement("Upload", typeof(UploadTaskParameters))]
+	public class UploadTask : CustomTask
+	{
+		/// <summary>
+		/// Parameters for this task
+		/// </summary>
+		UploadTaskParameters Parameters;
+
+		/// <summary>
+		/// Constructor
+		/// </summary>
+		/// <param name="InParameters">Parameters for this task</param>
+		public UploadTask(UploadTaskParameters InParameters)
+		{
+			Parameters = InParameters;
+		}
+
+		/// <summary>
+		/// Execute the task.
+		/// </summary>
+		/// <param name="Job">Information about the current job</param>
+		/// <param name="BuildProducts">Set of build products produced by this node.</param>
+		/// <param name="TagNameToFileSet">Mapping from tag names to the set of files they include</param>
+		public override void Execute(JobContext Job, HashSet<FileReference> BuildProducts, Dictionary<string, HashSet<FileReference>> TagNameToFileSet)
+		{
+			// Find all the input files
+			List<FileReference> Files;
+			if(Parameters.Files == null)
+			{
+				Files = DirectoryReference.EnumerateFiles(Parameters.FromDir, "*", System.IO.SearchOption.AllDirectories).ToList();
+			}
+			else
+			{
+				Files = ResolveFilespec(Parameters.FromDir, Parameters.Files, TagNameToFileSet).ToList();
+			}
+
+			// Create the jupiter tree
+			Log.TraceInformation("Uploading {0} files to {1}...", Files.Count, Parameters.JupiterUrl);
+
+			JupiterFileTree FileTree = new JupiterFileTree(Parameters.FromDir, Parameters.LimitMemoryUsage);
+			foreach (FileReference File in Files)
+			{
+				FileTree.AddFile(File);
+			}
+
+			Dictionary<string, object> Metadata = new Dictionary<string, object>
+			{
+				{"ArchiveType", Parameters.ArchiveType},
+				{"Project", Parameters.ProjectName},
+				{"Branch", Parameters.Branch},
+				{"Changelist", Parameters.Changelist},
+			};
+
+			if (Parameters.AdditionalMetadata != null)
+			{
+				string[] KV = Parameters.AdditionalMetadata.Split(';');
+				foreach (string Option in KV)
+				{
+					int SeparatorIndex = Option.IndexOf('=');
+					if (SeparatorIndex == -1)
+						continue;
+
+					string Key = Option.Substring(0, SeparatorIndex);
+					string Value = Option.Substring(SeparatorIndex + 1);
+
+					Metadata[Key] = Value;
+				}
+			}
+			// Upload the tree to Jupiter
+			Task<Dictionary<FileReference, List<string>>> Task = FileTree.UploadToJupiter(Parameters.JupiterUrl, Parameters.JupiterNamespace, Parameters.JupiterKey, Metadata);
+			Task.Wait();
+			Dictionary<FileReference, List<string>> Mapping = Task.Result;
+
+			// Debug output of which files mapped to which blobs, can be useful to determine which files are constantly being uploaded
+			// Json.Save(FileReference.Combine(AutomationTool.CommandUtils.RootDirectory, "JupiterUpload.json"), Mapping);
+		}
+
+		/// <summary>
+		/// Output this task out to an XML writer.
+		/// </summary>
+		public override void Write(XmlWriter Writer)
+		{
+			Write(Writer, Parameters);
+		}
+
+		/// <summary>
+		/// Find all the tags which are used as inputs to this task
+		/// </summary>
+		/// <returns>The tag names which are read by this task</returns>
+		public override IEnumerable<string> FindConsumedTagNames()
+		{
+			return FindTagNamesFromFilespec(Parameters.Files);
+		}
+
+		public override IEnumerable<string> FindProducedTagNames()
+		{
+			return new string[] { };
+		}
+	}
+}
diff --git a/Engine/Source/Programs/AutomationTool/Scripts/AutomationScripts.Automation.csproj b/Engine/Source/Programs/AutomationTool/Scripts/AutomationScripts.Automation.csproj
--- a/Engine/Source/Programs/AutomationTool/Scripts/AutomationScripts.Automation.csproj
+++ b/Engine/Source/Programs/AutomationTool/Scripts/AutomationScripts.Automation.csproj
@@ -78,6 +78,7 @@
     <Compile Include="BenchmarkBuild\BenchmarkRunEditorTask.cs" />
     <Compile Include="BenchmarkBuild\BenchmarkSingleCompileTask.cs" />
     <Compile Include="BenchmarkBuild\BenchmarkTaskBase.cs" />
+    <Compile Include="DownloadJupiterBuild.cs" />
     <Compile Include="CheckPerforceCase.cs" />
     <Compile Include="FinalizeInstalledBuild.cs" />
     <Compile Include="BuildHlslcc.Automation.cs" />
diff --git a/Engine/Source/Programs/AutomationTool/Scripts/DownloadJupiterBuild.cs b/Engine/Source/Programs/AutomationTool/Scripts/DownloadJupiterBuild.cs
new file mode 100644
--- /dev/null
+++ b/Engine/Source/Programs/AutomationTool/Scripts/DownloadJupiterBuild.cs
@@ -0,0 +1,28 @@
+// Copyright Epic Games, Inc. All Rights Reserved.
+
+using AutomationTool;
+using System.Threading.Tasks;
+using Tools.DotNETCommon;
+using Tools.DotNETCommon.Jupiter;
+
+namespace BuildScripts.Automation
+{
+	[Help("Downloads a build from Jupiter ")]
+	class DownloadJupiterBuild : BuildCommand
+	{
+		
+		public override void ExecuteBuild()
+		{
+			string Key = ParseRequiredStringParam("key");
+			string Namespace = ParseRequiredStringParam("namespace");
+			string JupiterUrl = ParseRequiredStringParam("url");
+			DirectoryReference TargetDirectory = ParseRequiredDirectoryReferenceParam("targetdirectory");
+
+			JupiterFileTree Tree = new JupiterFileTree(TargetDirectory, true);
+			
+			FileReference LocalManifest = FileReference.Combine(TargetDirectory, "Jupiter-Manifest.json");
+			Task DownloadTask = Tree.DownloadFromJupiter(LocalManifest, JupiterUrl, Namespace, Key);
+			DownloadTask.Wait();
+		}
+	}
+}
diff --git a/Engine/Source/Programs/DotNETCommon/DotNETUtilities/DotNETUtilities.csproj b/Engine/Source/Programs/DotNETCommon/DotNETUtilities/DotNETUtilities.csproj
--- a/Engine/Source/Programs/DotNETCommon/DotNETUtilities/DotNETUtilities.csproj
+++ b/Engine/Source/Programs/DotNETCommon/DotNETUtilities/DotNETUtilities.csproj
@@ -45,6 +45,7 @@
   <ItemGroup>
     <Reference Include="System" />
     <Reference Include="System.Core" />
+    <Reference Include="System.Net.Http" />
     <Reference Include="System.Xml.Linq" />
     <Reference Include="System.Data.DataSetExtensions" />
     <Reference Include="Microsoft.CSharp" />
@@ -87,6 +88,8 @@
     <Compile Include="Json.cs" />
     <Compile Include="JsonObject.cs" />
     <Compile Include="JsonWriter.cs" />
+    <Compile Include="Jupiter\JupierUtils.cs" />
+    <Compile Include="Jupiter\JupiterTree.cs" />
     <Compile Include="Log.cs" />
     <Compile Include="LogIndentScope.cs" />
     <Compile Include="LogStatusScope.cs" />
diff --git a/Engine/Source/Programs/DotNETCommon/DotNETUtilities/Jupiter/JupierUtils.cs b/Engine/Source/Programs/DotNETCommon/DotNETUtilities/Jupiter/JupierUtils.cs
new file mode 100644
--- /dev/null
+++ b/Engine/Source/Programs/DotNETCommon/DotNETUtilities/Jupiter/JupierUtils.cs
@@ -0,0 +1,112 @@
+// Copyright Epic Games, Inc. All Rights Reserved.
+
+using System;
+using System.Collections.Generic;
+using System.IO;
+using System.Linq;
+using System.Security.Cryptography;
+using System.Text;
+using System.Threading.Tasks;
+
+namespace Tools.DotNETCommon.Jupiter
+{
+
+	public static class JupiterFileUtils
+	{
+		public static List<Tuple<byte[], byte[]>> ChunkFile(FileReference File)
+		{
+			List<Tuple<byte[], byte[]>> Chunks = new List<Tuple<byte[], byte[]>>();
+			const int ChunkSize = 10 * 1024 * 1024; // chunk size in bytes
+			const int BlockSize = 4 * 1024; // the size of each block we read
+			byte[] FileBlock = new byte[BlockSize];
+
+			using (FileStream FileStream = new FileStream(File.FullName, FileMode.Open))
+			{
+				bool MoreDataAvailable = true;
+				while (MoreDataAvailable)
+				{
+					int TotalReadCountOfBytes = 0;
+					long RemainingFileSize = FileStream.Length - FileStream.Position;
+					if (RemainingFileSize == 0)
+					{
+						MoreDataAvailable = false;
+						continue;
+					}
+
+					int CurrentChunkSize = Math.Min((int)RemainingFileSize, ChunkSize);
+					byte[] Chunk = new byte[CurrentChunkSize];
+
+					while (TotalReadCountOfBytes < CurrentChunkSize)
+					{
+						int ReadCountOfBytes = FileStream.Read(FileBlock, 0, BlockSize);
+						if (ReadCountOfBytes == 0)
+						{
+							MoreDataAvailable = false;
+							break;
+						}
+						Array.Copy(FileBlock, 0, Chunk, TotalReadCountOfBytes, ReadCountOfBytes);
+
+						TotalReadCountOfBytes += ReadCountOfBytes;
+					}
+
+					byte[] Hash = SHA1Utils.GetSHA1(Chunk);
+					Chunks.Add(new Tuple<byte[], byte[]>(Hash, Chunk));
+				}
+			}
+
+			return Chunks;
+		}
+	}
+
+	public static class SHA1Utils
+	{
+		public static byte[] GetSHA1(byte[] payload)
+		{
+			using (SHA1 sha1 = SHA1.Create())
+			{
+				return sha1.ComputeHash(payload);
+			}
+		}
+
+		public static string GetSHA1AsString(byte[] payload)
+		{
+			byte[] sha1 = GetSHA1(payload);
+
+			return FormatAsHexString(sha1);
+		}
+
+		private static int GetHexVal(char hex)
+		{
+			return hex - (hex < 58 ? 48 : hex < 97 ? 55 : 87);
+		}
+
+		public static byte[] ToSha1FromString(string hex)
+		{
+			if (hex.Length % 2 == 1)
+			{
+				throw new Exception("The binary key cannot have an odd number of digits");
+			}
+
+			byte[] arr = new byte[hex.Length >> 1];
+
+			for (int i = 0; i < hex.Length >> 1; ++i)
+			{
+				arr[i] = (byte)((GetHexVal(hex[i << 1]) << 4) + GetHexVal(hex[(i << 1) + 1]));
+			}
+
+			return arr;
+		}
+
+		public static string FormatAsHexString(byte[] sha1Hash)
+		{
+			StringBuilder sb = new StringBuilder();
+			foreach (byte b in sha1Hash)
+			{
+				string hex = b.ToString("X2");
+				sb.Append(hex);
+			}
+
+			return sb.ToString();
+		}
+	}
+}
diff --git a/Engine/Source/Programs/DotNETCommon/DotNETUtilities/Jupiter/JupiterTree.cs b/Engine/Source/Programs/DotNETCommon/DotNETUtilities/Jupiter/JupiterTree.cs
new file mode 100644
--- /dev/null
+++ b/Engine/Source/Programs/DotNETCommon/DotNETUtilities/Jupiter/JupiterTree.cs
@@ -0,0 +1,673 @@
+// Copyright Epic Games, Inc. All Rights Reserved.
+
+using System;
+using System.Collections.Generic;
+using System.IO;
+using System.Linq;
+using System.Net.Http;
+using System.Text;
+using System.Threading.Tasks;
+
+namespace Tools.DotNETCommon.Jupiter
+{
+	public class JupiterTree
+	{
+		private readonly JupiterTreeContentProvider ContentProvider;
+		private byte[] LastTreeHash = null; // cache the calculated tree hash, reset the cache by setting to null
+
+		public JupiterTree(JupiterTreeContentProvider InContentProvider)
+		{
+			ContentProvider = InContentProvider;
+			ContentProvider.OnContentChanged += OnContentChanged;
+		}
+
+		private void OnContentChanged(object Sender, EventArgs EventArgs)
+		{
+			// invalidate the tree hash when the contents change
+			LastTreeHash = null;
+		}
+
+
+		public List<JupiterTree> Trees { get; } = new List<JupiterTree>();
+
+		public IEnumerable<string> ContentHashes
+		{
+			get { return ContentProvider.GetHashes(); }
+		}
+
+		public Task<byte[]> GetContent(string BlobHash)
+		{
+			return ContentProvider.GetContent(BlobHash);
+		}
+
+		public void AddTree(JupiterTree Tree)
+		{
+			LastTreeHash = null;
+			Trees.Add(Tree);
+		}
+
+		public void AddTrees(IEnumerable<JupiterTree> InTrees)
+		{
+			LastTreeHash = null;
+			Trees.AddRange(InTrees);
+		}
+
+		public byte[] CalculateTreeHash()
+		{
+			if (LastTreeHash != null)
+				return LastTreeHash;
+
+			List<byte> Bytes = new List<byte>();
+			foreach (JupiterTree Tree in Trees)
+			{
+				Bytes.AddRange(Tree.CalculateTreeHash());
+			}
+
+			foreach (string Hash in ContentProvider.GetHashes())
+			{
+				Bytes.AddRange(SHA1Utils.ToSha1FromString(Hash));
+			}
+
+			LastTreeHash = SHA1Utils.GetSHA1(Bytes.ToArray());
+			return LastTreeHash;
+		}
+
+		public List<JupiterTree> GetAllTrees()
+		{
+			List<JupiterTree> AllTrees = new List<JupiterTree>();
+			DoGetAllTrees(AllTrees);
+			return AllTrees;
+		}
+
+		private void DoGetAllTrees(List<JupiterTree> OutTrees)
+		{
+			OutTrees.Add(this);
+			foreach (JupiterTree Tree in Trees)
+			{
+				Tree.DoGetAllTrees(OutTrees);
+			}
+		}
+
+		public async Task UploadToJupiter(string JupiterUrl, string JupiterNamespace, string JupiterTreeKey, Dictionary<string, object> Metadata = null)
+		{
+			string TreeHashString = SHA1Utils.FormatAsHexString(CalculateTreeHash());
+			List<JupiterTree> AllTrees = GetAllTrees();
+			using (HttpClient JupiterClient = new HttpClient {BaseAddress = new Uri(JupiterUrl)})
+			{
+				Log.TraceInformation("Creating a new tree root in Jupiter with id: {0}", JupiterTreeKey);
+
+				if (Metadata == null)
+				{
+					Metadata = new Dictionary<string, object>();
+				}
+
+				// add a source field if it doesn't already exist to make it known were this tree came from
+				if (!Metadata.ContainsKey("source"))
+				{
+					Metadata["source"] = "BuildGraph";
+				}
+
+				{
+					var TreeRoot = new
+					{
+						TreeReferenceKey = JupiterTreeKey,
+						TreeHash = TreeHashString,
+						Metadata = Metadata,
+					};
+					string PutTreeRootString = Json.Serialize(TreeRoot, JsonSerializeOptions.None);
+					StringContent PutTreeRootContent = new StringContent(PutTreeRootString, Encoding.UTF8, "application/json");
+					// upload tree information
+					HttpResponseMessage PutTreeResult = await JupiterClient.PutAsync(string.Format("api/v1/c/tree-root/{0}", JupiterNamespace), PutTreeRootContent);
+					if (!PutTreeResult.IsSuccessStatusCode)
+					{
+						string ErrorMsg = await PutTreeResult.Content.ReadAsStringAsync();
+						throw new Exception(string.Format("Failed to create tree root in Jupiter. Response: {0}", ErrorMsg));
+					}
+				}
+
+				Log.TraceInformation("Uploading \"{0}\" trees to Jupiter", AllTrees.Count);
+
+				{
+					// upload the trees we have
+					var PutTreesRequest = from Tree in AllTrees
+						select new
+						{
+							Hash = SHA1Utils.FormatAsHexString(Tree.CalculateTreeHash()),
+							Trees = Tree.Trees.Select(JupiterTree =>
+								SHA1Utils.FormatAsHexString(JupiterTree.CalculateTreeHash())),
+							Blobs = Tree.ContentHashes,
+						};
+					string PutTreeString = Json.Serialize(PutTreesRequest, JsonSerializeOptions.None);
+					StringContent PutTreeContent = new StringContent(PutTreeString, Encoding.UTF8, "application/json");
+					HttpResponseMessage PutTreeContentResult = await JupiterClient.PutAsync(string.Format("api/v1/c/tree/{0}", JupiterNamespace), PutTreeContent);
+					if (!PutTreeContentResult.IsSuccessStatusCode)
+					{
+						string ErrorMsg = await PutTreeContentResult.Content.ReadAsStringAsync();
+						throw new Exception(string.Format("Failed to upload one or more trees to Jupiter. Response: {0}", ErrorMsg));
+					}
+				}
+
+				List<string> FoundBlobs;
+				{
+					// determine which blobs are already present and thus do not need to be uploaded
+					List<string> BlobIdentifiers = AllTrees.SelectMany(Tree => Tree.ContentHashes).ToList();
+
+					var FilterBlobRequest = new
+					{
+						Operations = from Blob in BlobIdentifiers
+							select new
+							{
+								Namespace = JupiterNamespace,
+								Id = Blob,
+								Op = "HEAD"
+							}
+					};
+
+					string FilterBlobString = Json.Serialize(FilterBlobRequest, JsonSerializeOptions.None);
+					StringContent FilterBlobContent = new StringContent(FilterBlobString, Encoding.UTF8, "application/json");
+					HttpResponseMessage FilterBlobResponse = await JupiterClient.PostAsync("api/v1/s", FilterBlobContent);
+					if (!FilterBlobResponse.IsSuccessStatusCode)
+					{
+						string ErrorMsg = await FilterBlobResponse.Content.ReadAsStringAsync();
+						throw new Exception(string.Format("Failed to determine which blobs were already present in Jupiter. Response: {0}", ErrorMsg));
+					}
+
+					string ResponseString = await FilterBlobResponse.Content.ReadAsStringAsync();
+					List<string> UnknownBlobs = Json.Deserialize<List<string>>(ResponseString);
+					FoundBlobs = BlobIdentifiers.Where(Blob => !UnknownBlobs.Contains(Blob)).ToList();
+
+					Log.TraceInformation("Determined that build consist of \"{0}\" chunks of which \"{1}\" where already present. Thus uploading \"{2}\" blobs", BlobIdentifiers.Count, FoundBlobs.Count, BlobIdentifiers.Count - FoundBlobs.Count);
+				}
+
+				{
+					// upload the blobs
+					List<Task> SubmitTasks = new List<Task>();
+					foreach (JupiterTree JupiterTree in AllTrees)
+					{
+						foreach (string BlobHash in JupiterTree.ContentHashes)
+						{
+							// check if blob is already present
+							if (FoundBlobs.Contains(BlobHash))
+								continue;
+
+							SubmitTasks.Add(Task.Run(async () =>
+								{
+									ByteArrayContent Content = new ByteArrayContent(await JupiterTree.GetContent(BlobHash));
+									Content.Headers.Remove("Content-Type");
+									Content.Headers.Add("Content-Type", "application/octet-stream");
+									HttpResponseMessage PutBlobResultResult = await JupiterClient.PutAsync(string.Format("api/v1/s/{0}/{1}", JupiterNamespace, BlobHash), Content);
+
+									if (!PutBlobResultResult.IsSuccessStatusCode)
+									{
+										string ErrorMsg = await PutBlobResultResult.Content.ReadAsStringAsync();
+										throw new Exception(string.Format(
+											"Failed to upload content blob to Jupiter. Response: {0}", ErrorMsg));
+									}
+								}
+							));
+						}
+					}
+
+					Task.WaitAll(SubmitTasks.ToArray());
+				}
+
+				// verify the build upload
+				{
+					HttpResponseMessage FinalizeTreeResult = await JupiterClient.PutAsync(string.Format("api/v1/c/tree-root/{0}/{1}/finalize", JupiterNamespace, JupiterTreeKey), null);
+					if (!FinalizeTreeResult.IsSuccessStatusCode)
+					{
+						string ErrorMsg = await FinalizeTreeResult.Content.ReadAsStringAsync();
+						throw new Exception(string.Format("Jupiter tree upload verification failed. Response: {0}",
+							ErrorMsg));
+					}
+				}
+
+				Log.TraceInformation("Build upload and verification complete.");
+			}
+		}
+
+		public static async Task<JupiterTree> FromJupiterKey(string JupiterUrl, string JupiterNamespace, string JupiterTreeKey)
+		{
+			using (HttpClient JupiterClient = new HttpClient {BaseAddress = new Uri(JupiterUrl)})
+			{
+				Log.TraceInformation("Downloading tree with key \"{0}\" from Jupiter at {1} in namespace {2}", JupiterTreeKey, JupiterUrl, JupiterNamespace);
+
+				string TopTreeHash;
+
+				{
+					// download the tree root
+					HttpResponseMessage GetTreeRootResult = await JupiterClient.GetAsync(string.Format("api/v1/c/tree-root/{0}/{1}", JupiterNamespace, JupiterTreeKey));
+
+					string GetTreeRootResultString = await GetTreeRootResult.Content.ReadAsStringAsync();
+					if (!GetTreeRootResult.IsSuccessStatusCode)
+					{
+						throw new Exception(string.Format("Failed to download tree root with key {0}. Response: {1}", JupiterTreeKey, GetTreeRootResultString));
+					}
+
+					TreeRootContents TreeRoot = Json.Deserialize<TreeRootContents>(GetTreeRootResultString);
+
+					TopTreeHash = TreeRoot.treeHash;
+				}
+
+				Dictionary<string, TreeContents> TreeMapping;
+				{
+					// Flatten this tree, and download the tree descriptions of all those trees
+					HttpResponseMessage GetFlattenTreeResult = await JupiterClient.GetAsync(string.Format("api/v1/c/tree/{0}/{1}/flattend", JupiterNamespace, TopTreeHash));
+
+					string GetFlattenTreeResultString = await GetFlattenTreeResult.Content.ReadAsStringAsync();
+					if (!GetFlattenTreeResult.IsSuccessStatusCode)
+					{
+						throw new Exception(string.Format("Failed to download flattened tree {0}. Response: {1}", TopTreeHash, GetFlattenTreeResultString));
+					}
+
+					FlattendTreeContents FlattenedTree = Json.Deserialize<FlattendTreeContents>(GetFlattenTreeResultString);
+
+					List<string> TreesToDownload = new List<string>(FlattenedTree.allTrees);
+					TreesToDownload.Add(TopTreeHash);
+					TreeMapping = new Dictionary<string, TreeContents>();
+					foreach (string Tree in TreesToDownload)
+					{
+						// Flatten this tree, and download the tree descriptions of all those trees
+						HttpResponseMessage GetTreeResult = await JupiterClient.GetAsync(string.Format("api/v1/c/tree/{0}/{1}", JupiterNamespace, Tree));
+
+						string GetTreeResultString = await GetTreeResult.Content.ReadAsStringAsync();
+						if (!GetTreeResult.IsSuccessStatusCode)
+						{
+							throw new Exception(string.Format("Failed to download tree {0}. Response: {1}", Tree, GetTreeResultString));
+						}
+
+
+						TreeContents TreeContents = Json.Deserialize<TreeContents>(GetTreeResultString);
+						TreeMapping.Add(Tree, TreeContents);
+					}
+				}
+
+				{
+					return CreateTreeFromTreeDescriptions(JupiterUrl, JupiterNamespace, TopTreeHash, TreeMapping);
+				}
+			}
+		}
+
+		private static JupiterTree CreateTreeFromTreeDescriptions(string JupiterUrl, string JupiterNamespace, string RootTreeHash, Dictionary<string, TreeContents> TreeDescriptions)
+		{
+			TreeContents RootTreeContents = TreeDescriptions[RootTreeHash];
+			JupiterIoContentProvider ContentProvider = new JupiterIoContentProvider();
+			ContentProvider.AddContentReference(JupiterUrl, JupiterNamespace, RootTreeContents.blobs);
+			JupiterTree NewTree = new JupiterTree(ContentProvider);
+
+			foreach (string Tree in RootTreeContents.trees)
+			{
+				NewTree.AddTree(CreateTreeFromTreeDescriptions(JupiterUrl, JupiterNamespace, Tree, TreeDescriptions));
+			}
+
+			return NewTree;
+		}
+
+		public class TreeRootContents
+		{
+			public string treeHash { get; set; }
+			public Dictionary<string, object> metadata { get; set; }
+		}
+
+		public class FlattendTreeContents
+		{
+			public string[] allTrees { get; set; }
+
+			public string[] allBlobs { get; set; }
+		}
+
+		public class TreeContents
+		{
+			public string treeHash { get; set; }
+			public List<string> trees { get; set; }
+			public List<string> blobs { get; set; }
+		}
+	}
+
+	public class JupiterIoContentProvider : JupiterTreeContentProvider
+	{
+		public string BaseUrl { get; private set; }
+		public List<string> Blobs { get; private set; }
+		public string JupiterNamespace { get; private set; }
+
+		public override async Task<byte[]> GetContent(string Sha1)
+		{
+			using (HttpClient JupiterClient = new HttpClient {BaseAddress = new Uri(BaseUrl)})
+			{
+				HttpResponseMessage GetContentsResult = await JupiterClient.GetAsync(string.Format("api/v1/s/{0}/{1}", JupiterNamespace, Sha1));
+
+				if (!GetContentsResult.IsSuccessStatusCode)
+				{
+					string ErrorMsg = await GetContentsResult.Content.ReadAsStringAsync();
+					throw new Exception(string.Format("Failed to download blob {0}. Response: {1}", Sha1, ErrorMsg));
+				}
+
+				byte[] Data = await GetContentsResult.Content.ReadAsByteArrayAsync();
+
+				return Data;
+			}
+		}
+
+		public override IEnumerable<string> GetHashes()
+		{
+			return Blobs;
+		}
+
+		public override event EventHandler OnContentChanged;
+
+		public void AddContentReference(string JupiterUrl, string InJupiterNamespace, List<string> InBlobs)
+		{
+			BaseUrl = JupiterUrl;
+			JupiterNamespace = InJupiterNamespace;
+			Blobs = InBlobs;
+
+			OnContentChanged?.Invoke(this, new EventArgs());
+		}
+
+	}
+
+
+	public class JupiterStreamedContentProvider : JupiterTreeContentProvider
+	{
+		private class ChunkDescription
+		{
+			public long Offset;
+			public long Size;
+		}
+
+		private FileReference File { get; set; }
+		private Dictionary<string, ChunkDescription> Chunks;
+
+		public override event EventHandler OnContentChanged;
+
+		public void SetFileSource(FileReference InFile)
+		{
+			File = InFile;
+			Chunks = PopulateChunks(File);
+
+			OnContentChanged?.Invoke(this, new EventArgs());
+		}
+
+
+		public override Task<byte[]> GetContent(string Sha1)
+		{
+			ChunkDescription Chunk;
+			if (Chunks.TryGetValue(Sha1, out Chunk))
+			{
+				using (FileStream FileStream = new FileStream(File.FullName, FileMode.Open, FileAccess.Read, FileShare.Read))
+				using (BinaryReader BinaryReader = new BinaryReader(FileStream))
+				{
+					FileStream.Position = Chunk.Offset;
+					return Task.FromResult(BinaryReader.ReadBytes((int)Chunk.Size));
+				}
+			}
+			throw new ArgumentOutOfRangeException(nameof(Sha1), "Unable to find any matching chunk for hash: " + Sha1);
+		}
+
+		public override IEnumerable<string> GetHashes()
+		{
+			return Chunks.Keys;
+		}
+
+		private static Dictionary<string, ChunkDescription> PopulateChunks(FileReference File)
+		{
+			Dictionary<string, ChunkDescription> NewChunks = new Dictionary<string, ChunkDescription>();
+			List<Tuple<byte[], byte[]>> FileChunks = JupiterFileUtils.ChunkFile(File);
+			long CurrentOffset = 0;
+			foreach (Tuple<byte[], byte[]> FileChunk in FileChunks)
+			{
+				long ChunkLength = FileChunk.Item2.LongLength;
+				ChunkDescription ChunkDescription = new ChunkDescription { Offset = CurrentOffset, Size = ChunkLength };
+				CurrentOffset += ChunkLength;
+				string Hash = SHA1Utils.FormatAsHexString(FileChunk.Item1);
+				NewChunks.Add(Hash, ChunkDescription);
+			}
+
+			return NewChunks;
+		}
+
+	}
+
+	public class JupiterInMemoryContentProvider : JupiterTreeContentProvider
+	{
+		private readonly Dictionary<string, byte[]> ContentBlobs = new Dictionary<string, byte[]>();
+
+		public void AddContent(byte[] Sha1, byte[] Content)
+		{
+			if (Sha1.Length != 20)
+				throw new ArgumentException("Sha1 argument was not 20 bytes, did you really specify a hash of the contents?", nameof(Sha1));
+
+			string Hash = SHA1Utils.FormatAsHexString(Sha1);
+			ContentBlobs.Add(Hash, Content);
+			OnContentChanged?.Invoke(this, new EventArgs());
+		}
+
+		public override Task<byte[]> GetContent(string Sha1)
+		{
+			return Task.FromResult(ContentBlobs[Sha1]);
+		}
+
+		public override IEnumerable<string> GetHashes()
+		{
+			return ContentBlobs.Keys;
+		}
+
+		public override event EventHandler OnContentChanged;
+	}
+
+	public abstract class JupiterTreeContentProvider
+	{
+		public abstract Task<byte[]> GetContent(string Sha1);
+		public abstract IEnumerable<string> GetHashes();
+		public abstract event EventHandler OnContentChanged;
+	}
+
+
+	/// <summary>
+	/// A merkel tree built from a set of files ready to be uploaded to Jupiter
+	/// </summary>
+	public class JupiterFileTree
+	{
+		private readonly DirectoryReference BaseDir;
+		private readonly List<FileReference> Files = new List<FileReference>();
+		private readonly bool DeferReadingFiles;
+
+		public JupiterFileTree(DirectoryReference InBaseDir, bool InDeferReadingFiles)
+		{
+			BaseDir = InBaseDir;
+			DeferReadingFiles = InDeferReadingFiles;
+		}
+
+		public void AddFile(FileReference File)
+		{
+			Files.Add(File);
+		}
+
+		public async Task<Dictionary<FileReference, List<string>>> UploadToJupiter(string JupiterUrl, string JupiterNamespace, string JupiterTreeKey, Dictionary<string, object> Metadata = null)
+		{
+			Manifest Manifest = new Manifest(BaseDir);
+
+			JupiterTree JupiterTopTree = new JupiterTree(new JupiterInMemoryContentProvider());
+
+			List<JupiterTree> FileTrees = new List<JupiterTree>();
+			Dictionary<FileReference, List<string>> FileToChunkMapping = new Dictionary<FileReference, List<string>>();
+			foreach (FileReference File in Files)
+			{
+				JupiterTreeContentProvider ContentProvider;
+				if (DeferReadingFiles)
+				{
+					// deferred reading means we avoid keeping the file in memory, instead reading the contents from disk as we are uploading it.
+					JupiterStreamedContentProvider JupiterStreamedContentProvider = new JupiterStreamedContentProvider();
+					JupiterStreamedContentProvider.SetFileSource(File);
+					ContentProvider = JupiterStreamedContentProvider;
+				}
+				else
+				{
+					// In memory use-case - just read all files into memory and hash them.
+					JupiterInMemoryContentProvider InMemoryContentProvider = new JupiterInMemoryContentProvider();
+					List<Tuple<byte[], byte[]>> FileChunks = JupiterFileUtils.ChunkFile(File);
+					foreach (Tuple<byte[], byte[]> Chunk in FileChunks)
+					{
+						byte[] Sha1 = Chunk.Item1;
+						byte[] Content = Chunk.Item2;
+						InMemoryContentProvider.AddContent(Sha1, Content);
+					}
+
+					ContentProvider = InMemoryContentProvider;
+
+				}
+
+				JupiterTree CurrentTree = new JupiterTree(ContentProvider);
+
+				// populate the manifest as we have prepared the content providers
+				Manifest.AddFile(File, CurrentTree.CalculateTreeHash());
+				FileTrees.Add(CurrentTree);
+				FileToChunkMapping.Add(File, ContentProvider.GetHashes().ToList());
+			}
+
+
+			// The manifest is the first node in the tree so that we can find it later
+			JupiterTopTree.AddTree(Manifest.AsTree());
+			JupiterTopTree.AddTrees(FileTrees);
+
+			await JupiterTopTree.UploadToJupiter(JupiterUrl, JupiterNamespace, JupiterTreeKey, Metadata);
+
+			return FileToChunkMapping;
+		}
+
+		public async Task<List<FileReference>> DownloadFromJupiter(FileReference LocalManifestPath, string JupiterUrl, string JupiterNamespace, string JupiterTreeKey, IProgress<Tuple<float, FileReference>> Progress = null)
+		{
+			Manifest LocalManifest = null;
+			if (LocalManifestPath.ToFileInfo().Exists)
+			{
+				LocalManifest = Manifest.CreateFromFileReference(BaseDir, LocalManifestPath);
+			}
+			JupiterTree Tree = await JupiterTree.FromJupiterKey(JupiterUrl, JupiterNamespace, JupiterTreeKey);
+
+			JupiterTree ManifestTree = Tree.Trees.First();
+			Manifest NewManifest = Manifest.FromTree(BaseDir, ManifestTree);
+
+			int CountOfFileTrees = Tree.Trees.Count - 1; // all trees except for the first represents files
+			JupiterTree[] FileTrees = new JupiterTree[CountOfFileTrees];
+			Tree.Trees.CopyTo(1, FileTrees, 0, CountOfFileTrees);
+
+			Log.TraceInformation("Build {0} consisted of {1} files. Downloading.", JupiterTreeKey, CountOfFileTrees);
+
+			List<FileReference> FilesWritten = new List<FileReference>();
+			int CountOfFilesWritten = 0;
+			foreach (JupiterTree FileTree in FileTrees)
+			{
+				string TreeHash = SHA1Utils.FormatAsHexString(FileTree.CalculateTreeHash());
+
+				KeyValuePair<string, string> ManifestFileReference = NewManifest.Files.First(Pair => Pair.Value == TreeHash);
+				FileReference LocalFileReference = FileReference.Combine(BaseDir, ManifestFileReference.Key);
+
+				bool FileAlreadyPresent = false;
+				string LocalHash;
+				// if we have a local manifest which contains this same file with the same hash we can ignore that file
+				if (LocalManifest != null && LocalManifest.Files.TryGetValue(ManifestFileReference.Key, out LocalHash))
+				{
+					if (FileReference.Exists(LocalFileReference) && LocalHash == TreeHash)
+					{
+						Log.TraceInformation("File \"{0}\" already present with hash \"{1}\". Skipping.", LocalFileReference, LocalHash);
+						FileAlreadyPresent = true;
+					}
+				}
+
+				if (!FileAlreadyPresent)
+				{
+					Log.TraceInformation("Creating file: {0}", LocalFileReference);
+
+					foreach (string ContentHash in FileTree.ContentHashes)
+					{
+						byte[] Content = await FileTree.GetContent(ContentHash);
+						DirectoryReference.CreateDirectory(LocalFileReference.Directory);
+						FileReference.WriteAllBytes(LocalFileReference, Content);
+					}
+				}
+
+				FilesWritten.Add(LocalFileReference);
+				Progress?.Report(new Tuple<float, FileReference>(++CountOfFilesWritten / (float)CountOfFileTrees, LocalFileReference));
+			}
+
+			NewManifest.Save(LocalManifestPath);
+			return FilesWritten;
+		}
+
+		public class Manifest
+		{
+			private readonly DirectoryReference BaseDir;
+
+			// ReSharper disable once MemberCanBePrivate.Global , fastjson requires this to be public so it can be serialized
+			public class ManifestData
+			{
+				public string SchemaVersion = "v1"; // added in case we need to update this schema in the future
+				public Dictionary<string, string> Files = new Dictionary<string, string>();
+			}
+			private readonly ManifestData Data;
+
+			public Dictionary<string, string> Files
+			{
+				get { return Data.Files; }
+			}
+
+			public Manifest(DirectoryReference InBaseDir)
+			{
+				BaseDir = InBaseDir;
+				Data = new ManifestData();
+			}
+
+			private Manifest(DirectoryReference InBaseDir, ManifestData InData)
+			{
+				BaseDir = InBaseDir;
+				Data = InData;
+			}
+
+			public void AddFile(FileReference File, byte[] TreeHash)
+			{
+				// we escape the forward slashes to make it valid json keys
+				Files.Add(File.MakeRelativeTo(BaseDir).Replace(@"\", @"\\"), SHA1Utils.FormatAsHexString(TreeHash));
+			}
+
+			public JupiterTree AsTree()
+			{
+				string JsonFiles = Json.Serialize(Data, JsonSerializeOptions.None);
+				byte[] ManifestBytes = Encoding.ASCII.GetBytes(JsonFiles);
+				byte[] Hash = SHA1Utils.GetSHA1(ManifestBytes);
+				JupiterInMemoryContentProvider ContentProvider = new JupiterInMemoryContentProvider();
+				ContentProvider.AddContent(Hash, ManifestBytes);
+				return new JupiterTree(ContentProvider);
+			}
+
+			public static Manifest FromTree(DirectoryReference BaseDir, JupiterTree ManifestTree)
+			{
+				// the first hash of the tree is a json serialized form of the manifest, see AsTree
+				string FirstHash = ManifestTree.ContentHashes.First();
+				byte[] Content = ManifestTree.GetContent(FirstHash).Result;
+				string JsonContent = Encoding.ASCII.GetString(Content);
+				ManifestData Data = Json.Deserialize<ManifestData>(JsonContent);
+				return new Manifest(BaseDir, Data);
+			}
+
+			public static Manifest CreateFromFileReference(DirectoryReference BaseDir, FileReference FileReference)
+			{
+				ManifestData Data = Json.Load<ManifestData>(FileReference);
+				return new Manifest(BaseDir, Data);
+			}
+
+			public void Save(FileReference LocalManifestReference)
+			{
+				DirectoryReference.CreateDirectory(LocalManifestReference.Directory);
+
+				Dictionary<string, string> EscapedFiles = Data.Files.ToDictionary(Pair => Json.EscapeString(Pair.Key), Pair => Pair.Value);
+
+				ManifestData EscapedData = new ManifestData()
+				{
+					SchemaVersion = Data.SchemaVersion,
+					Files = EscapedFiles
+				};
+				Json.Save(LocalManifestReference, EscapedData);
+			}
+		}
+	}
+}
\ No newline at end of file
