- Renamed a poorly named PlatformSDK function for AutoSDK (GetDesiredVersion -> GetMainVersion)
- Now allow AutoSDK directory name and internal version to be separate (Android only using this for now)
- Some code cleanup with associated functions

diff --git a/Engine/Source/Programs/AutomationTool/Turnkey/FileSource.cs b/Engine/Source/Programs/AutomationTool/Turnkey/FileSource.cs
--- a/Engine/Source/Programs/AutomationTool/Turnkey/FileSource.cs
+++ b/Engine/Source/Programs/AutomationTool/Turnkey/FileSource.cs
@@ -277,7 +277,7 @@ namespace Turnkey
 				if (Type == SourceType.AutoSdk)
 				{
 					UInt64 AutoSdkVersion;
-					if (SDK.TryConvertVersionToInt(SDK.GetDesiredVersion(), out AutoSdkVersion))
+					if (SDK.TryConvertVersionToInt(SDK.GetMainVersion(), out AutoSdkVersion))
 					{
 						Best = Sdks.FirstOrDefault(x => { UInt64 IntVer; return SDK.TryConvertVersionToInt(x.Version, out IntVer) && IntVer == AutoSdkVersion; });
 					}
diff --git a/Engine/Source/Programs/AutomationTool/Turnkey/SdkInfo.cs b/Engine/Source/Programs/AutomationTool/Turnkey/SdkInfo.cs
--- a/Engine/Source/Programs/AutomationTool/Turnkey/SdkInfo.cs
+++ b/Engine/Source/Programs/AutomationTool/Turnkey/SdkInfo.cs
@@ -61,7 +61,7 @@ namespace Turnkey
 					foreach (DirectoryInfo Dir in PlatformDir.EnumerateDirectories())
 					{
 						// check for valid version number (auto SDK must match version exactly)
-						if (string.Compare(Dir.Name, SDK.GetRequiredSDKString()) == 0)
+						if (string.Compare(Dir.Name, SDK.GetAutoSDKDirectoryForMasterVersion()) == 0)
 						{
 							// make sure it actually has buits in it
 							if (File.Exists(Path.Combine(Dir.FullName, "setup.bat")))
diff --git a/Engine/Source/Programs/DotNETCommon/BuildUtilities/UEBuildPlatformSDK.cs b/Engine/Source/Programs/DotNETCommon/BuildUtilities/UEBuildPlatformSDK.cs
--- a/Engine/Source/Programs/DotNETCommon/BuildUtilities/UEBuildPlatformSDK.cs
+++ b/Engine/Source/Programs/DotNETCommon/BuildUtilities/UEBuildPlatformSDK.cs
@@ -121,7 +121,7 @@ namespace Tools.DotNETCommon
 			if (PlatformSupportsAutoSDKs())
 			{
 				// we don't need ask the platform for the installed version because we know the AutoSDK version has to match the DesiredVersion if it was setup properly
-				//AutoSDKVersion = (HasRequiredAutoSDKInstalled() == SDKStatus.Valid) ? GetDesiredVersion() : null;
+				//AutoSDKVersion = (HasRequiredAutoSDKInstalled() == SDKStatus.Valid) ? GetMainVersion() : null;
 				AutoSDKVersion = GetInstalledSDKVersion();
 			}
 			else
@@ -162,11 +162,28 @@ namespace Tools.DotNETCommon
 		/// Return the SDK version that the platform wants to use (AutoSDK dir must match this, full SDKs can be in a valid range)
 		/// </summary>
 		/// <returns></returns>
-		public virtual string GetDesiredVersion()
+		public abstract string GetMainVersion();
+
+		/// <summary>
+		/// Gets the valid string range of Sdk versions. TryConvertVersionToInt() will need to succeed to make this usable for range checks
+		/// </summary>
+		/// <param name="MinVersion">Smallest version allowed</param>
+		/// <param name="MaxVersion">Largest version allowed (inclusive)</param>
+		/// <returns>True if the versions are valid, false if the platform is unable to convert its versions into an integer</returns>
+		public abstract void GetValidVersionRange(out string MinVersion, out string MaxVersion);
+
+		/// <summary>
+		/// For a platform that doesn't use properly named AutoSDK directories, the directory name may not be convertible to an integer,
+		/// and IsVersionValid checks could fail when checking AutoSDK version for an exact match. GetMainVersion() would return the 
+		/// proper, integer-convertible version number of the SDK inside of the directory returned by GetAutoSDKDirectoryForMasterVersion()
+		/// </summary>
+		/// <returns></returns>
+		public virtual string GetAutoSDKDirectoryForMasterVersion()
 		{
-			return GetRequiredSDKString();
+			return GetMainVersion();
 		}
 
+
 		/// <summary>
 		/// Gets the valid (integer) range of Sdk versions. Must be an integer to easily check a range vs a particular version
 		/// </summary>
@@ -188,19 +205,6 @@ namespace Tools.DotNETCommon
 			return false;
 		}
 
-		/// <summary>
-		/// Gets the valid string range of Sdk versions. TryConvertVersionToInt() will need to succeed to make this usable for range checks
-		/// </summary>
-		/// <param name="MinVersion">Smallest version allowed</param>
-		/// <param name="MaxVersion">Largest version allowed (inclusive)</param>
-		/// <returns>True if the versions are valid, false if the platform is unable to convert its versions into an integer</returns>
-		public virtual void GetValidVersionRange(out string MinVersion, out string MaxVersion)
-		{
-			// @todo turnkey make this abstract?
-			// by default, use GetDesiredVersion() to min and max
-			MinVersion = MaxVersion = GetDesiredVersion();
-		}
-
 		// Let platform override behavior to determine if a version is a valid (useful for non-numeric versions)
 		public virtual bool IsVersionValidInternal(string Version, bool bForAutoSDK)
 		{
@@ -218,7 +222,7 @@ namespace Tools.DotNETCommon
 			}
 
 			UInt64 DesiredVersion;
-			if (!TryConvertVersionToInt(GetDesiredVersion(), out DesiredVersion))
+			if (!TryConvertVersionToInt(GetMainVersion(), out DesiredVersion))
 			{
 				return false;
 			}
@@ -363,15 +367,6 @@ namespace Tools.DotNETCommon
 			return !IsAutoSDKDestructive() || !HasAnyManualInstall();
 		}
 
-		/// <summary>
-		/// Returns SDK string as required by the platform
-		/// </summary>
-		/// <returns>Valid SDK string</returns>
-		public virtual string GetRequiredSDKString()
-		{
-			return "";
-		}
-
 		/// <summary>
 		/// Gets the version number of the SDK setup script itself.  The version in the base should ALWAYS be the master revision from the last refactor.
 		/// If you need to force a rebuild for a given platform, override this for the given platform.
@@ -525,7 +520,7 @@ namespace Tools.DotNETCommon
 		{
 			if (PlatformSupportsAutoSDKs() && HasAutoSDKSystemEnabled())
 			{
-				String InstalledSDKVersionString = GetRequiredSDKString();
+				String InstalledSDKVersionString = GetAutoSDKDirectoryForMasterVersion();
 				String PlatformSDKRoot = GetPathToPlatformAutoSDKs();
                 if (!Directory.Exists(PlatformSDKRoot))
                 {
@@ -877,7 +872,7 @@ namespace Tools.DotNETCommon
 					bool bEnvVarFileExists = File.Exists(EnvVarFile);
 
 					string CurrentSDKString;
-					if (bEnvVarFileExists && GetCurrentlyInstalledSDKString(AutoSDKRoot, out CurrentSDKString) && CurrentSDKString == GetRequiredSDKString() && bScriptVersionMatches)
+					if (bEnvVarFileExists && GetCurrentlyInstalledSDKString(AutoSDKRoot, out CurrentSDKString) && CurrentSDKString == GetAutoSDKDirectoryForMasterVersion() && bScriptVersionMatches)
 					{
 						return SDKStatus.Valid;
 					}
@@ -985,22 +980,22 @@ namespace Tools.DotNETCommon
 					// delete Manifest file to avoid multiple uninstalls
 					InvalidateCurrentlyInstalledAutoSDK();
 
-					if (!RunAutoSDKHooks(AutoSDKRoot, GetRequiredSDKString(), SDKHookType.Install, false))
+					if (!RunAutoSDKHooks(AutoSDKRoot, GetAutoSDKDirectoryForMasterVersion(), SDKHookType.Install, false))
 					{
-						Log.TraceLog("Failed to install required SDK {0}.  Attemping to uninstall", GetRequiredSDKString());
-						RunAutoSDKHooks(AutoSDKRoot, GetRequiredSDKString(), SDKHookType.Uninstall, false);
+						Log.TraceLog("Failed to install required SDK {0}.  Attemping to uninstall", GetAutoSDKDirectoryForMasterVersion());
+						RunAutoSDKHooks(AutoSDKRoot, GetAutoSDKDirectoryForMasterVersion(), SDKHookType.Uninstall, false);
 						return;
 					}
 
 					string EnvVarFile = Path.Combine(AutoSDKRoot, SDKEnvironmentVarsFile);
 					if (!File.Exists(EnvVarFile))
 					{
-						Log.TraceLog("Installation of required SDK {0}.  Did not generate Environment file {1}", GetRequiredSDKString(), EnvVarFile);
-						RunAutoSDKHooks(AutoSDKRoot, GetRequiredSDKString(), SDKHookType.Uninstall, false);
+						Log.TraceLog("Installation of required SDK {0}.  Did not generate Environment file {1}", GetAutoSDKDirectoryForMasterVersion(), EnvVarFile);
+						RunAutoSDKHooks(AutoSDKRoot, GetAutoSDKDirectoryForMasterVersion(), SDKHookType.Uninstall, false);
 						return;
 					}
 
-					SetCurrentlyInstalledAutoSDKString(GetRequiredSDKString());
+					SetCurrentlyInstalledAutoSDKString(GetAutoSDKDirectoryForMasterVersion());
 					SetLastRunAutoSDKScriptVersion(GetRequiredScriptVersionString());
 				}
 
@@ -1032,7 +1027,7 @@ namespace Tools.DotNETCommon
 			// load environment variables from current SDK
 			if (!SetupEnvironmentFromAutoSDK(PlatformSDKRoot))
 			{
-				Log.TraceLog("Failed to load environment from required SDK {0}", GetRequiredSDKString());
+				Log.TraceLog("Failed to load environment from required SDK {0}", GetAutoSDKDirectoryForMasterVersion());
 				InvalidateCurrentlyInstalledAutoSDK();
 				return SDKStatus.Invalid;
 			}
diff --git a/Engine/Source/Programs/UnrealBuildTool/Configuration/UEBuildPlatform.cs b/Engine/Source/Programs/UnrealBuildTool/Configuration/UEBuildPlatform.cs
--- a/Engine/Source/Programs/UnrealBuildTool/Configuration/UEBuildPlatform.cs
+++ b/Engine/Source/Programs/UnrealBuildTool/Configuration/UEBuildPlatform.cs
@@ -162,15 +162,6 @@ namespace UnrealBuildTool
 			return UEBuildPlatform.GetSDK(Platform).HasRequiredSDKsInstalled();
 		}
 
-		/// <summary>
-		/// Returns SDK string as required by the platform
-		/// </summary>
-		/// <returns>Valid SDK string</returns>
-		public string GetRequiredSDKString()
-		{
-			return UEBuildPlatform.GetSDK(Platform).GetRequiredSDKString();
-		}
-
 		/// <summary>
 		/// Whether this platform requires specific Visual Studio version.
 		/// </summary>
diff --git a/Engine/Source/Programs/UnrealBuildTool/Platform/Android/AndroidPlatformSDK.cs b/Engine/Source/Programs/UnrealBuildTool/Platform/Android/AndroidPlatformSDK.cs
--- a/Engine/Source/Programs/UnrealBuildTool/Platform/Android/AndroidPlatformSDK.cs
+++ b/Engine/Source/Programs/UnrealBuildTool/Platform/Android/AndroidPlatformSDK.cs
@@ -11,13 +11,13 @@ namespace UnrealBuildTool
 {
 	class AndroidPlatformSDK : UEBuildPlatformSDK
 	{
-		public override string GetRequiredSDKString()
+		public override string GetMainVersion()
 		{
-			return "-23";
+			return "r21b";
 		}
-		public override string GetDesiredVersion()
+		public override string GetAutoSDKDirectoryForMasterVersion()
 		{
-			return "r21b";
+			return "-23";
 		}
 
 		public override void GetValidVersionRange(out string MinVersion, out string MaxVersion)
@@ -27,7 +27,7 @@ namespace UnrealBuildTool
 		}
 
 		// the version we prefer
-		private string RecommendedVersion = "21b";
+		private string RecommendedVersion = "r21b";
 
 		public override string GetInstalledSDKVersion()
 		{
diff --git a/Engine/Source/Programs/UnrealBuildTool/Platform/Linux/LinuxPlatformSDK.cs b/Engine/Source/Programs/UnrealBuildTool/Platform/Linux/LinuxPlatformSDK.cs
--- a/Engine/Source/Programs/UnrealBuildTool/Platform/Linux/LinuxPlatformSDK.cs
+++ b/Engine/Source/Programs/UnrealBuildTool/Platform/Linux/LinuxPlatformSDK.cs
@@ -10,7 +10,7 @@ namespace UnrealBuildTool
 {
 	class LinuxPlatformSDK : UEBuildPlatformSDK
 	{
-		public override string GetDesiredVersion()
+		public override string GetMainVersion()
 		{
 			return "v16_clang-9.0.1-centos7";
 		}
@@ -114,14 +114,14 @@ namespace UnrealBuildTool
 			string SDKRoot = Environment.GetEnvironmentVariable(SDKRootEnvVar);
 			if (!String.IsNullOrEmpty(SDKRoot))
 			{
-				string AutoSDKPath = Path.Combine(SDKRoot, "Host" + BuildHostPlatform.Current.Platform, TargetPlatformName, GetDesiredVersion(), LinuxPlatform.DefaultHostArchitecture);
+				string AutoSDKPath = Path.Combine(SDKRoot, "Host" + BuildHostPlatform.Current.Platform, TargetPlatformName, GetAutoSDKDirectoryForMasterVersion(), LinuxPlatform.DefaultHostArchitecture);
 				if (DirectoryReference.Exists(new DirectoryReference(AutoSDKPath)))
 				{
 					return AutoSDKPath;
 				}
 			}
 
-			string InTreeSDKPath = Path.Combine(LinuxPlatformSDK.GetInTreeSDKRoot().FullName, GetDesiredVersion(), LinuxPlatform.DefaultHostArchitecture);
+			string InTreeSDKPath = Path.Combine(LinuxPlatformSDK.GetInTreeSDKRoot().FullName, GetMainVersion(), LinuxPlatform.DefaultHostArchitecture);
 			if (DirectoryReference.Exists(new DirectoryReference(InTreeSDKPath)))
 			{
 				return InTreeSDKPath;
@@ -130,15 +130,6 @@ namespace UnrealBuildTool
 			return null;
 		}
 
-		/// <summary>
-		/// Returns SDK string as required by the platform
-		/// </summary>
-		/// <returns>Valid SDK string</returns>
-		public override string GetRequiredSDKString()
-		{
-			return GetDesiredVersion();
-		}
-
 		protected override String GetRequiredScriptVersionString()
 		{
 			return "3.0";
@@ -209,7 +200,7 @@ namespace UnrealBuildTool
 				DirectoryReference InTreeSDKVersionRoot = GetInTreeSDKRoot();
 				if (InTreeSDKVersionRoot != null)
 				{
-					DirectoryReference InTreeSDKVersionPath = DirectoryReference.Combine(InTreeSDKVersionRoot, GetDesiredVersion());
+					DirectoryReference InTreeSDKVersionPath = DirectoryReference.Combine(InTreeSDKVersionRoot, GetMainVersion());
 					if (DirectoryReference.Exists(InTreeSDKVersionPath))
 					{
 						MultiArchRoot = InTreeSDKVersionPath.FullName;
diff --git a/Engine/Source/Programs/UnrealBuildTool/Platform/Lumin/LuminPlatformSDK.cs b/Engine/Source/Programs/UnrealBuildTool/Platform/Lumin/LuminPlatformSDK.cs
--- a/Engine/Source/Programs/UnrealBuildTool/Platform/Lumin/LuminPlatformSDK.cs
+++ b/Engine/Source/Programs/UnrealBuildTool/Platform/Lumin/LuminPlatformSDK.cs
@@ -11,17 +11,15 @@ namespace UnrealBuildTool
 	class LuminPlatformSDK : UEBuildPlatformSDK
 	{
 
-		public override string GetDesiredVersion()
+		public override string GetMainVersion()
 		{
 			return "0.23";
 		}
 
-		// not needed just returning GetDesiredVersion()
-		// 		public override void GetValidVersionRange(out string MinVersion, out string MaxVersion)
-		// 		{
-		// 			MinVersion = "0.23";
-		// 			MaxVersion = "0.23";
-		// 		}
+ 		public override void GetValidVersionRange(out string MinVersion, out string MaxVersion)
+		{
+			MinVersion = MaxVersion = GetMainVersion();
+		}
 
 
 		public override string GetInstalledSDKVersion()
@@ -86,11 +84,6 @@ namespace UnrealBuildTool
 			return true;
 		}
 
-		public override string GetRequiredSDKString()
-		{
-			return GetDesiredVersion();
-		}
-
 		protected override String GetRequiredScriptVersionString()
 		{
 			return "Lumin_15";
diff --git a/Engine/Source/Programs/UnrealBuildTool/Platform/Mac/ApplePlatformSDK.cs b/Engine/Source/Programs/UnrealBuildTool/Platform/Mac/ApplePlatformSDK.cs
--- a/Engine/Source/Programs/UnrealBuildTool/Platform/Mac/ApplePlatformSDK.cs
+++ b/Engine/Source/Programs/UnrealBuildTool/Platform/Mac/ApplePlatformSDK.cs
@@ -12,7 +12,7 @@ namespace UnrealBuildTool
 {
 	internal class ApplePlatformSDK : UEBuildPlatformSDK
 	{
-		public override string GetDesiredVersion()
+		public override string GetMainVersion()
 		{
 			return "11.5";
 		}
diff --git a/Engine/Source/Programs/UnrealBuildTool/Platform/Windows/MicrosoftPlatformSDK.cs b/Engine/Source/Programs/UnrealBuildTool/Platform/Windows/MicrosoftPlatformSDK.cs
--- a/Engine/Source/Programs/UnrealBuildTool/Platform/Windows/MicrosoftPlatformSDK.cs
+++ b/Engine/Source/Programs/UnrealBuildTool/Platform/Windows/MicrosoftPlatformSDK.cs
@@ -10,7 +10,7 @@ namespace UnrealBuildTool
 {
 	internal class MicrosoftPlatformSDK : UEBuildPlatformSDK
 	{
-		public override string GetDesiredVersion()
+		public override string GetMainVersion()
 		{
 			// the current and previous versions of the SDK (technically, NX Add-On), with appended NEX version, we support
 			return "10.0.18632";
diff --git a/Engine/Source/Programs/UnrealBuildTool/System/PlatformExports.cs b/Engine/Source/Programs/UnrealBuildTool/System/PlatformExports.cs
--- a/Engine/Source/Programs/UnrealBuildTool/System/PlatformExports.cs
+++ b/Engine/Source/Programs/UnrealBuildTool/System/PlatformExports.cs
@@ -99,16 +99,6 @@ namespace UnrealBuildTool
 			return BuildPlatform.GetExcludedFolderNames().ToArray();
 		}
 
-		/// <summary>
-		/// Returns the respective platform sdk version string
-		/// </summary>
-		/// <param name="Platform">The target platform to query</param>
-		public static string GetRequiredSDKString(UnrealTargetPlatform Platform)
-		{
-			UEBuildPlatform BuildPlatform = UEBuildPlatform.GetBuildPlatform(Platform, false);
-			return BuildPlatform.GetRequiredSDKString();
-		}
-
 		/// <summary>
 		/// Check whether the given platform supports XGE
 		/// </summary>
