Log to the console when we don't have a valid SDK. For now, only warn if the SDK is installed with the wrong version.

diff --git a/Engine/Source/Programs/DotNETCommon/BuildUtilities/UEBuildPlatformSDK.cs b/Engine/Source/Programs/DotNETCommon/BuildUtilities/UEBuildPlatformSDK.cs
--- a/Engine/Source/Programs/DotNETCommon/BuildUtilities/UEBuildPlatformSDK.cs
+++ b/Engine/Source/Programs/DotNETCommon/BuildUtilities/UEBuildPlatformSDK.cs
@@ -1175,6 +1175,12 @@ namespace Tools.DotNETCommon
 				Msg.AppendLine("    RunUAT Turnkey -command=InstallSdk -platform={0}", PlatformName);
 
 				Log.WriteLine(LogEventType.Error, Options, Msg.ToString());
+
+				// If an SDK is installed but isn't at the right version, make sure the user knows.
+				if (!String.IsNullOrEmpty(ManualSDKVersion))
+				{
+					Log.TraceInformationOnce(Msg.ToString());
+				}
 			}
 		}
 
