- Updated UEBuildPlatform[SDK] to look in DDPI.ini and not bother setting up SDKs for platforms that aren't enabled (no need to print out errors for missing platform SDKs for platforms that aren't valid, like consoles on Mac)
- Streamlined the logs for bad SDK versions, especially when multiple platforms have bad SDKs
- Skip over non-enabled platforms in some Turnkey commands

diff --git a/Engine/Source/Programs/AutomationTool/Turnkey/Commands/ListPlatforms.cs b/Engine/Source/Programs/AutomationTool/Turnkey/Commands/ListPlatforms.cs
--- a/Engine/Source/Programs/AutomationTool/Turnkey/Commands/ListPlatforms.cs
+++ b/Engine/Source/Programs/AutomationTool/Turnkey/Commands/ListPlatforms.cs
@@ -31,6 +31,10 @@ namespace Turnkey.Commands
 
 				Platform Platform = Platform.Platforms[new TargetPlatformDescriptor(TargetPlatform)];
 				UEBuildPlatformSDK SDK = UEBuildPlatformSDK.GetSDKForPlatform(TargetPlatform.ToString());
+				if (SDK == null || !SDK.bIsSdkAllowedOnHost)
+				{
+					continue;
+				}
 
 				string ManualSDKVersion, AutoSDKVersion;
 				string MinAllowedVersion, MaxAllowedVersion;
diff --git a/Engine/Source/Programs/AutomationTool/Turnkey/TurnkeyManifest.cs b/Engine/Source/Programs/AutomationTool/Turnkey/TurnkeyManifest.cs
--- a/Engine/Source/Programs/AutomationTool/Turnkey/TurnkeyManifest.cs
+++ b/Engine/Source/Programs/AutomationTool/Turnkey/TurnkeyManifest.cs
@@ -113,6 +113,10 @@ namespace Turnkey
 			}
 
 			// NotNeedsExpansion now contains the expanded sources, and the sources that didn't need to be expanded
+
+			// now filter out platforms that are disabled on this host
+			NotNeedsExpansion = NotNeedsExpansion.FindAll(x => x.GetPlatforms().Any(y => UEBuildPlatformSDK.GetSDKForPlatform(y.ToString()).bIsSdkAllowedOnHost));
+
 			return NotNeedsExpansion;
 		}
 
diff --git a/Engine/Source/Programs/DotNETCommon/BuildUtilities/UEBuildPlatformSDK.cs b/Engine/Source/Programs/DotNETCommon/BuildUtilities/UEBuildPlatformSDK.cs
--- a/Engine/Source/Programs/DotNETCommon/BuildUtilities/UEBuildPlatformSDK.cs
+++ b/Engine/Source/Programs/DotNETCommon/BuildUtilities/UEBuildPlatformSDK.cs
@@ -43,7 +43,7 @@ namespace Tools.DotNETCommon
 		/// </summary>
 		/// <param name="SDK">SDK object</param>
 		/// <param name="PlatformName">Platform name for this SDK</param>
-		public static void RegisterSDKForPlatform(UEBuildPlatformSDK SDK, string PlatformName)
+		public static void RegisterSDKForPlatform(UEBuildPlatformSDK SDK, string PlatformName, bool bIsSdkAllowedOnHost)
 		{
 			// verify that neither platform or sdk were added before
 			if (SDKRegistry.Count(x => x.Key == PlatformName || x.Value == SDK) > 0)
@@ -53,12 +53,13 @@ namespace Tools.DotNETCommon
 
 			SDKRegistry.Add(PlatformName, SDK);
 
-			SDK.Init(PlatformName);
+			SDK.Init(PlatformName, bIsSdkAllowedOnHost);
 		}
 
-		private void Init(string InPlatformName)
+		private void Init(string InPlatformName, bool bInIsSdkAllowedOnHost)
 		{
 			PlatformName = InPlatformName;
+			bIsSdkAllowedOnHost = bInIsSdkAllowedOnHost;
 
 			// if the parent set up autosdk, the env vars will be wrong, but we can still get the manual SDK version from before it was setup
 			if (HasParentProcessSetupAutoSDK(out CachedManualSDKVersion))
@@ -104,6 +105,9 @@ namespace Tools.DotNETCommon
 		// String name of the platform (will match an UnrealTargetPlatform)
 		public string PlatformName;
 
+		// True if this Sdk is allowed to be used by this host - if not, we can skip a lot 
+		public bool bIsSdkAllowedOnHost;
+
 		public string GetInstalledVersion(out bool bIsAutoSDK)
 		{
 			bIsAutoSDK = HasSetupAutoSDK();
@@ -1128,11 +1132,13 @@ namespace Tools.DotNETCommon
 				}
 			}
 
-			// print all SDKs to log file
-			PrintSDKInfo(LogEventType.Log, LogFormatOptions.NoConsoleOutput);
+			// print all SDKs to log file (errors will print out later for builds and generateprojectfiles)
+			PrintSDKInfo(LogEventType.Log, LogFormatOptions.NoConsoleOutput, LogEventType.Verbose, LogFormatOptions.NoConsoleOutput);
 		}
-		
-		public virtual void PrintSDKInfo(LogEventType Verbosity = LogEventType.Console, LogFormatOptions Options = LogFormatOptions.None)
+
+		private static bool bHasShownTurnkey = false;
+		public virtual void PrintSDKInfo(LogEventType Verbosity = LogEventType.Console, LogFormatOptions Options = LogFormatOptions.None,
+			LogEventType ErrorVerbosity = LogEventType.Error, LogFormatOptions ErrorOptions = LogFormatOptions.None)
 		{
 			string ManualSDKVersion, AutoSDKVersion;
 			GetInstalledVersions(out ManualSDKVersion, out AutoSDKVersion);
@@ -1156,31 +1162,34 @@ namespace Tools.DotNETCommon
 				GetValidVersionRange(out MinVersionString, out MaxVersionString);
 
 				StringBuilder Msg = new StringBuilder();
-				Msg.AppendLine("Unable to find a valid SDK for {0}", PlatformName);
+				Msg.AppendFormat("Unable to find a valid SDK for {0}.", PlatformName);
 				if (ManualSDKVersion != null)
 				{
-					Msg.AppendLine("  Found Version: {0}", ManualSDKVersion);
+					Msg.AppendFormat(" Found Version: {0}.", ManualSDKVersion);
 				}
 
 				if (MinVersionString != MaxVersionString)
 				{
-					Msg.AppendLine("  SDK version needs to be between {0} and {1}", MinVersionString, MaxVersionString);
+					Msg.AppendLine(" Must be between {0} and {1}", MinVersionString, MaxVersionString);
 				}
 				else
 				{
-					Msg.AppendLine("  SDK version needs to be {0}", MinVersionString);
+					Msg.AppendLine(" Must be {0}", MinVersionString);
 				}
 
-				Msg.AppendLine("  If your Studio has it set up, you can run this command to find the SDK to install:");
-				Msg.AppendLine("    RunUAT Turnkey -command=InstallSdk -platform={0}", PlatformName);
-
-				Log.WriteLine(LogEventType.Error, Options, Msg.ToString());
-
-				// If an SDK is installed but isn't at the right version, make sure the user knows.
-				if (!String.IsNullOrEmpty(ManualSDKVersion))
+				if (!bHasShownTurnkey)
 				{
-					Log.TraceInformationOnce(Msg.ToString());
+					Msg.AppendLine("  If your Studio has it set up, you can run this command to find the SDK to install:");
+					Msg.AppendLine("    RunUAT Turnkey -command=InstallSdk -platform={0} -BestAvailable", PlatformName);
+
+					if ((ErrorOptions & LogFormatOptions.NoConsoleOutput) == LogFormatOptions.None)
+					{
+						bHasShownTurnkey = true;
+					}
 				}
+
+				// always print errors to the screen
+				Log.WriteLine(ErrorVerbosity, ErrorOptions, Msg.ToString());
 			}
 		}
 
diff --git a/Engine/Source/Programs/UnrealBuildTool/Configuration/UEBuildPlatform.cs b/Engine/Source/Programs/UnrealBuildTool/Configuration/UEBuildPlatform.cs
--- a/Engine/Source/Programs/UnrealBuildTool/Configuration/UEBuildPlatform.cs
+++ b/Engine/Source/Programs/UnrealBuildTool/Configuration/UEBuildPlatform.cs
@@ -46,9 +46,17 @@ namespace UnrealBuildTool
 		public UEBuildPlatform(UnrealTargetPlatform InPlatform, UEBuildPlatformSDK SDK)
 		{
 			Platform = InPlatform;
-			UEBuildPlatformSDK.RegisterSDKForPlatform(SDK, Platform.ToString());
 
-			SDK.ManageAndValidateSDK();
+			// check DDPI to see if the platform is enabled on this host platform
+			string IniPlatformName = ConfigHierarchy.GetIniPlatformName(Platform);
+			bool bIsEnabled = DataDrivenPlatformInfo.GetDataDrivenInfoForPlatform(IniPlatformName).bIsEnabled;
+
+			// set up the SDK if the platform is enabled
+			UEBuildPlatformSDK.RegisterSDKForPlatform(SDK, Platform.ToString(), bIsEnabled);
+			if (bIsEnabled)
+			{
+				SDK.ManageAndValidateSDK();
+			}
 		}
 
 		/// <summary>
diff --git a/Engine/Source/Programs/UnrealBuildTool/Modes/GenerateProjectFilesMode.cs b/Engine/Source/Programs/UnrealBuildTool/Modes/GenerateProjectFilesMode.cs
--- a/Engine/Source/Programs/UnrealBuildTool/Modes/GenerateProjectFilesMode.cs
+++ b/Engine/Source/Programs/UnrealBuildTool/Modes/GenerateProjectFilesMode.cs
@@ -125,6 +125,16 @@ namespace UnrealBuildTool
 				}
 			}
 
+			// print out any errors to the log
+			foreach (string PlatformName in UnrealTargetPlatform.GetValidPlatformNames())
+			{
+				UEBuildPlatformSDK SDK = UEBuildPlatformSDK.GetSDKForPlatform(PlatformName);
+				if (SDK != null && SDK.bIsSdkAllowedOnHost)
+				{
+					SDK.PrintSDKInfo(LogEventType.Verbose, LogFormatOptions.NoConsoleOutput, LogEventType.Error, LogFormatOptions.None);
+				}
+			}
+
 			// Create each project generator and run it
 			List<ProjectFileGenerator> Generators = new List<ProjectFileGenerator>();
 			foreach (ProjectFileFormat ProjectFileFormat in ProjectFileFormats.Distinct())
diff --git a/Engine/Source/Programs/UnrealBuildTool/Platform/Android/AndroidPlatformSDK.cs b/Engine/Source/Programs/UnrealBuildTool/Platform/Android/AndroidPlatformSDK.cs
--- a/Engine/Source/Programs/UnrealBuildTool/Platform/Android/AndroidPlatformSDK.cs
+++ b/Engine/Source/Programs/UnrealBuildTool/Platform/Android/AndroidPlatformSDK.cs
@@ -111,9 +111,9 @@ namespace UnrealBuildTool
 		}
 
 
-		public override void PrintSDKInfo(LogEventType Verbosity, LogFormatOptions Options)
+		public override void PrintSDKInfo(LogEventType Verbosity, LogFormatOptions Options, LogEventType ErrorVerbosity, LogFormatOptions ErrorOptions)
 		{
-			base.PrintSDKInfo(Verbosity, Options);
+			base.PrintSDKInfo(Verbosity, Options, ErrorVerbosity, ErrorOptions);
 
 			if (GetInstalledVersion() != RecommendedVersion)
 			{
diff --git a/Engine/Source/Programs/UnrealBuildTool/System/DataDrivenPlatformInfo.cs b/Engine/Source/Programs/UnrealBuildTool/System/DataDrivenPlatformInfo.cs
--- a/Engine/Source/Programs/UnrealBuildTool/System/DataDrivenPlatformInfo.cs
+++ b/Engine/Source/Programs/UnrealBuildTool/System/DataDrivenPlatformInfo.cs
@@ -22,7 +22,12 @@ namespace UnrealBuildTool
 			/// <summary>
 			/// Is the platform a confidential ("console-style") platform
 			/// </summary>
-			public bool bIsConfidential;
+			public bool bIsConfidential = false;
+
+			/// <summary>
+			/// Is the platform enabled on this host platform
+			/// </summary>
+			public bool bIsEnabled = false;
 
 			/// <summary>
 			/// Additional restricted folders for this platform.
@@ -90,9 +95,16 @@ namespace UnrealBuildTool
 
 							// slightly nasty bool parsing for bool values
 							string Temp;
-							if (ParsedSection.TryGetValue("bIsConfidential", out Temp) == false || ConfigHierarchy.TryParse(Temp, out NewInfo.bIsConfidential) == false)
+							if (ParsedSection.TryGetValue("bIsConfidential", out Temp))
+							{
+								ConfigHierarchy.TryParse(Temp, out NewInfo.bIsConfidential);
+							}
+
+							string HostKey = ConfigHierarchy.GetIniPlatformName(BuildHostPlatform.Current.Platform) + ":bIsEnabled";
+							string NormalKey = "bIsEnabled";
+							if (ParsedSection.TryGetValue(HostKey, out Temp) == true || ParsedSection.TryGetValue(NormalKey, out Temp) == true)
 							{
-								NewInfo.bIsConfidential = false;
+								ConfigHierarchy.TryParse(Temp, out NewInfo.bIsEnabled);
 							}
 
 							// get a list of additional restricted folders
