- Added file logging to GenerteProjectFiles (saved to Engine/Programs/UnrealBuildTool/Log_GPF.txt)
- Improved GPF's SDK error printouts, with no more red error text, and just a list of bad platform SDKs, and a link to the log file to get the detailed information if desired

diff --git a/Engine/Source/Programs/DotNETCommon/BuildUtilities/UEBuildPlatformSDK.cs b/Engine/Source/Programs/DotNETCommon/BuildUtilities/UEBuildPlatformSDK.cs
--- a/Engine/Source/Programs/DotNETCommon/BuildUtilities/UEBuildPlatformSDK.cs
+++ b/Engine/Source/Programs/DotNETCommon/BuildUtilities/UEBuildPlatformSDK.cs
@@ -320,6 +320,70 @@
 
 		#endregion
 
+		#region Print SDK Info
+		private static bool bHasShownTurnkey = false;
+		public virtual SDKStatus PrintSDKInfoAndReturnValidity(LogEventType Verbosity = LogEventType.Console, LogFormatOptions Options = LogFormatOptions.None,
+			LogEventType ErrorVerbosity = LogEventType.Error, LogFormatOptions ErrorOptions = LogFormatOptions.None)
+		{
+			string ManualSDKVersion, AutoSDKVersion;
+			GetInstalledVersions(out ManualSDKVersion, out AutoSDKVersion);
+
+			SDKStatus Validity = SDKStatus.Valid;
+
+			if (HasSetupAutoSDK())
+			{
+				string PlatformSDKRoot = GetPathToPlatformAutoSDKs();
+
+				UInt64 Ver;
+				TryConvertVersionToInt(AutoSDKVersion, out Ver);
+				Log.WriteLine(Verbosity, Options, "{0} using Auto SDK from: {1} 0x{2:X}", PlatformName, Path.Combine(PlatformSDKRoot, AutoSDKVersion), Ver);
+			}
+			else if (HasRequiredManualSDK() == SDKStatus.Valid)
+			{
+				Log.WriteLine(Verbosity, Options, "{0} using Manual SDK {1}", PlatformName, ManualSDKVersion);
+			}
+			else
+			{
+				Validity = SDKStatus.Invalid;
+
+				string MinVersionString, MaxVersionString;
+				GetValidVersionRange(out MinVersionString, out MaxVersionString);
+
+				StringBuilder Msg = new StringBuilder();
+				Msg.AppendFormat("Unable to find a valid SDK for {0}.", PlatformName);
+				if (ManualSDKVersion != null)
+				{
+					Msg.AppendFormat(" Found Version: {0}.", ManualSDKVersion);
+				}
+
+				if (MinVersionString != MaxVersionString)
+				{
+					Msg.AppendLine(" Must be between {0} and {1}", MinVersionString, MaxVersionString);
+				}
+				else
+				{
+					Msg.AppendLine(" Must be {0}", MinVersionString);
+				}
+
+				if (!bHasShownTurnkey)
+				{
+					Msg.AppendLine("  If your Studio has it set up, you can run this command to find the SDK to install:");
+					Msg.AppendLine("    RunUAT Turnkey -command=InstallSdk -platform={0} -BestAvailable", PlatformName);
+
+					if ((ErrorOptions & LogFormatOptions.NoConsoleOutput) == LogFormatOptions.None)
+					{
+						bHasShownTurnkey = true;
+					}
+				}
+
+				// always print errors to the screen
+				Log.WriteLine(ErrorVerbosity, ErrorOptions, Msg.ToString());
+			}
+
+			return Validity;
+		}
+
+		#endregion
 
 
 
@@ -1133,66 +1197,9 @@
 			}
 
 			// print all SDKs to log file (errors will print out later for builds and generateprojectfiles)
-			PrintSDKInfo(LogEventType.Log, LogFormatOptions.NoConsoleOutput, LogEventType.Verbose, LogFormatOptions.NoConsoleOutput);
+			PrintSDKInfoAndReturnValidity(LogEventType.Log, LogFormatOptions.NoConsoleOutput, LogEventType.Verbose, LogFormatOptions.NoConsoleOutput);
 		}
-
-		private static bool bHasShownTurnkey = false;
-		public virtual void PrintSDKInfo(LogEventType Verbosity = LogEventType.Console, LogFormatOptions Options = LogFormatOptions.None,
-			LogEventType ErrorVerbosity = LogEventType.Error, LogFormatOptions ErrorOptions = LogFormatOptions.None)
-		{
-			string ManualSDKVersion, AutoSDKVersion;
-			GetInstalledVersions(out ManualSDKVersion, out AutoSDKVersion);
-
-
-			if (HasSetupAutoSDK())
-			{
-				string PlatformSDKRoot = GetPathToPlatformAutoSDKs();
-
-				UInt64 Ver;
-				TryConvertVersionToInt(AutoSDKVersion, out Ver);
-				Log.WriteLine(Verbosity, Options, "{0} using Auto SDK from: {1} 0x{2:X}", PlatformName, Path.Combine(PlatformSDKRoot, AutoSDKVersion), Ver);
-			}
-			else if (HasRequiredManualSDK() == SDKStatus.Valid)
-			{
-				Log.WriteLine(Verbosity, Options, "{0} using Manual SDK {1}", PlatformName, ManualSDKVersion);
-			}
-			else
-			{
-				string MinVersionString, MaxVersionString;
-				GetValidVersionRange(out MinVersionString, out MaxVersionString);
-
-				StringBuilder Msg = new StringBuilder();
-				Msg.AppendFormat("Unable to find a valid SDK for {0}.", PlatformName);
-				if (ManualSDKVersion != null)
-				{
-					Msg.AppendFormat(" Found Version: {0}.", ManualSDKVersion);
-				}
-
-				if (MinVersionString != MaxVersionString)
-				{
-					Msg.AppendLine(" Must be between {0} and {1}", MinVersionString, MaxVersionString);
-				}
-				else
-				{
-					Msg.AppendLine(" Must be {0}", MinVersionString);
-				}
-
-				if (!bHasShownTurnkey)
-				{
-					Msg.AppendLine("  If your Studio has it set up, you can run this command to find the SDK to install:");
-					Msg.AppendLine("    RunUAT Turnkey -command=InstallSdk -platform={0} -BestAvailable", PlatformName);
-
-					if ((ErrorOptions & LogFormatOptions.NoConsoleOutput) == LogFormatOptions.None)
-					{
-						bHasShownTurnkey = true;
-					}
-				}
-
-				// always print errors to the screen
-				Log.WriteLine(ErrorVerbosity, ErrorOptions, Msg.ToString());
-			}
-		}
-
 		#endregion
+
 	}
 }
diff --git a/Engine/Source/Programs/UnrealBuildTool/Modes/BuildMode.cs b/Engine/Source/Programs/UnrealBuildTool/Modes/BuildMode.cs
--- a/Engine/Source/Programs/UnrealBuildTool/Modes/BuildMode.cs
+++ b/Engine/Source/Programs/UnrealBuildTool/Modes/BuildMode.cs
@@ -204,7 +204,7 @@ namespace UnrealBuildTool
 						}
 
 						// print out SDK info for only the platforms that are being compiled
-						UEBuildPlatformSDK.GetSDKForPlatform(TargetDesc.Platform.ToString()).PrintSDKInfo();
+						UEBuildPlatformSDK.GetSDKForPlatform(TargetDesc.Platform.ToString()).PrintSDKInfoAndReturnValidity();
 					}
 
 					// Get all the build options
diff --git a/Engine/Source/Programs/UnrealBuildTool/Modes/GenerateProjectFilesMode.cs b/Engine/Source/Programs/UnrealBuildTool/Modes/GenerateProjectFilesMode.cs
--- a/Engine/Source/Programs/UnrealBuildTool/Modes/GenerateProjectFilesMode.cs
+++ b/Engine/Source/Programs/UnrealBuildTool/Modes/GenerateProjectFilesMode.cs
@@ -8,13 +8,14 @@ using System.Reflection;
 using System.Text;
 using System.Threading.Tasks;
 using Tools.DotNETCommon;
+using System.Diagnostics;
 
 namespace UnrealBuildTool
 {
 	/// <summary>
 	/// Generates project files for one or more projects
 	/// </summary>
-	[ToolMode("GenerateProjectFiles", ToolModeOptions.XmlConfig | ToolModeOptions.BuildPlatforms | ToolModeOptions.SingleInstance)]
+	[ToolMode("GenerateProjectFiles", ToolModeOptions.XmlConfig | ToolModeOptions.BuildPlatforms | ToolModeOptions.SingleInstance | ToolModeOptions.UseStartupTraceListener)]
 	class GenerateProjectFilesMode : ToolMode
 	{
 		/// <summary>
@@ -58,6 +59,14 @@ namespace UnrealBuildTool
 			// Apply the XML config to this class
 			XmlConfig.ApplyTo(this);
 
+			// set up logging (taken from BuildMode)
+			FileReference LogFile = FileReference.Combine(UnrealBuildTool.EngineProgramSavedDirectory, "UnrealBuildTool", "Log_GPF.txt");
+			// find the StartupTraceListener in the listeners that was added super early on
+			StartupTraceListener StartupListener = Trace.Listeners.OfType<StartupTraceListener>().First();
+			TextWriterTraceListener LogTraceListener = Log.AddFileWriter("DefaultLogTraceListener", LogFile);
+			StartupListener.CopyTo(LogTraceListener);
+			Trace.Listeners.Remove(StartupListener);
+
 			// Parse rocket-specific arguments.
 			FileReference ProjectFile;
 			TryParseProjectFileArgument(Arguments, out ProjectFile);
@@ -126,14 +135,28 @@ namespace UnrealBuildTool
 			}
 
 			// print out any errors to the log
+			List<string> BadPlatformNames = new List<string>();
+			Log.TraceLog("\n---   SDK INFO START   ---");
 			foreach (string PlatformName in UnrealTargetPlatform.GetValidPlatformNames())
 			{
 				UEBuildPlatformSDK SDK = UEBuildPlatformSDK.GetSDKForPlatform(PlatformName);
 				if (SDK != null && SDK.bIsSdkAllowedOnHost)
 				{
-					SDK.PrintSDKInfo(LogEventType.Verbose, LogFormatOptions.NoConsoleOutput, LogEventType.Error, LogFormatOptions.None);
+					// print out the info to the log, and if it's invalid, remember it
+					SDKStatus Validity = SDK.PrintSDKInfoAndReturnValidity(LogEventType.Verbose, LogFormatOptions.NoConsoleOutput, LogEventType.Warning, LogFormatOptions.NoConsoleOutput);
+					if (Validity == SDKStatus.Invalid)
+					{
+						BadPlatformNames.Add(PlatformName);
+					}
 				}
 			}
+			if (BadPlatformNames.Count > 0)
+			{
+				Log.TraceInformationOnce("\nSome Platforms were skipped due to invalid SDK setup: {0}.\nSee the log file for detailed information ({1})\n\n", string.Join(", ", BadPlatformNames), Log.OutputFile);
+				Log.TraceInformation("");
+			}
+			Log.TraceLog("---   SDK INFO END   ---");
+			Log.TraceLog("");
 
 			// Create each project generator and run it
 			List<ProjectFileGenerator> Generators = new List<ProjectFileGenerator>();
diff --git a/Engine/Source/Programs/UnrealBuildTool/Platform/Android/AndroidPlatformSDK.cs b/Engine/Source/Programs/UnrealBuildTool/Platform/Android/AndroidPlatformSDK.cs
--- a/Engine/Source/Programs/UnrealBuildTool/Platform/Android/AndroidPlatformSDK.cs
+++ b/Engine/Source/Programs/UnrealBuildTool/Platform/Android/AndroidPlatformSDK.cs
@@ -111,14 +111,16 @@ namespace UnrealBuildTool
 		}
 
 
-		public override void PrintSDKInfo(LogEventType Verbosity, LogFormatOptions Options, LogEventType ErrorVerbosity, LogFormatOptions ErrorOptions)
+		public override SDKStatus PrintSDKInfoAndReturnValidity(LogEventType Verbosity, LogFormatOptions Options, LogEventType ErrorVerbosity, LogFormatOptions ErrorOptions)
 		{
-			base.PrintSDKInfo(Verbosity, Options, ErrorVerbosity, ErrorOptions);
+			SDKStatus Validity = base.PrintSDKInfoAndReturnValidity(Verbosity, Options, ErrorVerbosity, ErrorOptions);
 
 			if (GetInstalledVersion() != RecommendedVersion)
 			{
 				Log.WriteLine(Verbosity, Options, "Note: Android toolchain NDK {0} recommended", RecommendedVersion);
 			}
+
+			return Validity;
 		}
 
 
