Reworked Automation assembly loading in netcore again, using LoadFrom could cause multiple copies of the same assembly to be loaded (once by path and once by identity) causing very confusing errors as the same types were no longer considered the same (ptr difference).

We now always load by identity and hint to our assembly resolver were to expect to find assemblies that are not in our Engine.

diff --git a/Engine/Source/Programs/AutomationTool/AutomationUtils/ScriptCompiler.cs b/Engine/Source/Programs/AutomationTool/AutomationUtils/ScriptCompiler.cs
--- a/Engine/Source/Programs/AutomationTool/AutomationUtils/ScriptCompiler.cs
+++ b/Engine/Source/Programs/AutomationTool/AutomationUtils/ScriptCompiler.cs
@@ -475,7 +475,8 @@ namespace AutomationTool
 				try
 				{
 #if NET_CORE
-					Assembly Assembly = Assembly.LoadFile(AssemblyLocation.FullName);
+					AssemblyUtils.AddFileToAssemblyCache(AssemblyLocation.FullName);
+					Assembly Assembly = AppDomain.CurrentDomain.Load(AssemblyName.GetAssemblyName(AssemblyLocation.FullName));
 #else
 					Assembly Assembly = AppDomain.CurrentDomain.Load(AssemblyName.GetAssemblyName(AssemblyLocation.FullName));
 #endif
diff --git a/Engine/Source/Programs/DotNETCommon/DotNETUtilities/AssemblyUtils.cs b/Engine/Source/Programs/DotNETCommon/DotNETUtilities/AssemblyUtils.cs
--- a/Engine/Source/Programs/DotNETCommon/DotNETUtilities/AssemblyUtils.cs
+++ b/Engine/Source/Programs/DotNETCommon/DotNETUtilities/AssemblyUtils.cs
@@ -99,25 +99,30 @@ namespace Tools.DotNETCommon
 			// Initialize our cache of assemblies by enumerating all files in the given folder.
 			foreach (string DiscoveredAssembly in Directory.EnumerateFiles(RootDirectory, Pattern, SearchOption.AllDirectories))
 			{
-				string AssemblyName = Path.GetFileNameWithoutExtension(DiscoveredAssembly);
-				DateTime AssemblyLastWriteTime = File.GetLastWriteTimeUtc(DiscoveredAssembly);
-				if (AssemblyLocationCache.ContainsKey(AssemblyName))
-				{
-					// We already have this assembly in our cache. Only replace it if the discovered file is newer (to avoid stale assemblies breaking stuff).
-					if (AssemblyLastWriteTime > AssemblyWriteTimes[AssemblyName])
-					{
-						AssemblyLocationCache[AssemblyName] = DiscoveredAssembly;
-						AssemblyWriteTimes[AssemblyName] = AssemblyLastWriteTime;
-					}
-				}
-				else
+				AddFileToAssemblyCache(DiscoveredAssembly);
+			}
+
+		}
+
+		public static void AddFileToAssemblyCache(string AssemblyPath)
+		{
+			string AssemblyName = Path.GetFileNameWithoutExtension(AssemblyPath);
+			DateTime AssemblyLastWriteTime = File.GetLastWriteTimeUtc(AssemblyPath);
+			if (AssemblyLocationCache.ContainsKey(AssemblyName))
+			{
+				// We already have this assembly in our cache. Only replace it if the discovered file is newer (to avoid stale assemblies breaking stuff).
+				if (AssemblyLastWriteTime > AssemblyWriteTimes[AssemblyName])
 				{
-					// This is the first copy of this assembly ... add it to our cache.
-					AssemblyLocationCache.Add(AssemblyName, DiscoveredAssembly);
-					AssemblyWriteTimes.Add(AssemblyName, AssemblyLastWriteTime);
+					AssemblyLocationCache[AssemblyName] = AssemblyPath;
+					AssemblyWriteTimes[AssemblyName] = AssemblyLastWriteTime;
 				}
 			}
-
+			else
+			{
+				// This is the first copy of this assembly ... add it to our cache.
+				AssemblyLocationCache.Add(AssemblyName, AssemblyPath);
+				AssemblyWriteTimes.Add(AssemblyName, AssemblyLastWriteTime);
+			}
 		}
 
 		// Map of assembly name to path on disk
