Updated AutoSDK to allow unprivileged execution of the install/uninstall hooks when possible

diff --git a/Engine/Source/Programs/DotNETCommon/BuildUtilities/UEBuildPlatformSDK.cs b/Engine/Source/Programs/DotNETCommon/BuildUtilities/UEBuildPlatformSDK.cs
--- a/Engine/Source/Programs/DotNETCommon/BuildUtilities/UEBuildPlatformSDK.cs
+++ b/Engine/Source/Programs/DotNETCommon/BuildUtilities/UEBuildPlatformSDK.cs
@@ -713,6 +713,25 @@ namespace Tools.DotNETCommon
 			}
 		}
 
+		/// <summary>
+		/// Whether the hook must be run with administrator privileges.
+		/// </summary>
+		/// <param name="Hook">Hook for which to check the required privileges.</param>
+		/// <returns>true if the hook must be run with administrator privileges.</returns>
+		protected virtual bool DoesHookRequireAdmin(SDKHookType Hook)
+		{
+			return true;
+		}
+
+		private void LogAutoSDKHook(object Sender, DataReceivedEventArgs Args)
+		{
+			if (Args.Data != null)
+			{
+				LogFormatOptions Options = Log.OutputLevel >= LogEventType.Verbose ? LogFormatOptions.None : LogFormatOptions.NoConsoleOutput;
+				Log.WriteLine(LogEventType.Log, Options, Args.Data);
+			}
+		}
+
 		/// <summary>
 		/// Runs install/uninstall hooks for SDK
 		/// </summary>
@@ -744,21 +763,32 @@ namespace Tools.DotNETCommon
 					HookProcess.StartInfo.Arguments = "";
 					HookProcess.StartInfo.WindowStyle = ProcessWindowStyle.Hidden;
 
-					// seems to break the build machines?
-					//HookProcess.StartInfo.UseShellExecute = false;
-					//HookProcess.StartInfo.RedirectStandardOutput = true;
-					//HookProcess.StartInfo.RedirectStandardError = true;					
-
-//					using (ScopedTimer HookTimer = new ScopedTimer("Time to run hook: ", LogEventType.Log))
+					bool bHookRequiresAdmin = DoesHookRequireAdmin(Hook);
+					if (bHookRequiresAdmin)
 					{
-						//installers may require administrator access to succeed. so run as an admmin.
+						// installers may require administrator access to succeed. so run as an admin.
 						HookProcess.StartInfo.Verb = "runas";
+					}
+					else
+					{
+						HookProcess.StartInfo.UseShellExecute = false;
+						HookProcess.StartInfo.RedirectStandardOutput = true;
+						HookProcess.StartInfo.RedirectStandardError = true;
+						HookProcess.OutputDataReceived += LogAutoSDKHook;
+						HookProcess.ErrorDataReceived += LogAutoSDKHook;
+					}
+
+					//using (ScopedTimer HookTimer = new ScopedTimer("Time to run hook: ", LogEventType.Log))
+					{
 						HookProcess.Start();
+						if (!bHookRequiresAdmin)
+						{
+							HookProcess.BeginOutputReadLine();
+							HookProcess.BeginErrorReadLine();
+						}
 						HookProcess.WaitForExit();
 					}
 
-					//LogAutoSDK(HookProcess.StandardOutput.ReadToEnd());
-					//LogAutoSDK(HookProcess.StandardError.ReadToEnd());
 					if (HookProcess.ExitCode != 0)
 					{
 						Log.TraceLog("Hook exited uncleanly (returned {0}), considering it failed.", HookProcess.ExitCode);
diff --git a/Engine/Source/Programs/UnrealBuildTool/Platform/Linux/LinuxPlatformSDK.cs b/Engine/Source/Programs/UnrealBuildTool/Platform/Linux/LinuxPlatformSDK.cs
--- a/Engine/Source/Programs/UnrealBuildTool/Platform/Linux/LinuxPlatformSDK.cs
+++ b/Engine/Source/Programs/UnrealBuildTool/Platform/Linux/LinuxPlatformSDK.cs
@@ -82,6 +82,10 @@ namespace UnrealBuildTool
 			return true;
 		}
 
+		protected override bool DoesHookRequireAdmin(SDKHookType Hook)
+		{
+			return false;
+		}
 
 
 
