Add deprecation notices to Tools.DotNETCommon members.

#rnx

diff --git a/Engine/Source/Programs/DotNETCommon/DotNETUtilities/ArrayUtils.cs b/Engine/Source/Programs/DotNETCommon/DotNETUtilities/ArrayUtils.cs
--- a/Engine/Source/Programs/DotNETCommon/DotNETUtilities/ArrayUtils.cs
+++ b/Engine/Source/Programs/DotNETCommon/DotNETUtilities/ArrayUtils.cs
@@ -8,6 +8,7 @@ using System.Threading.Tasks;
 
 namespace Tools.DotNETCommon
 {
+	[Obsolete("Functionality in the Tools.DotNETCommon namespace is deprecated. Please reference the EpicGames.Common namespace and assembly instead.")]
 	static class ArrayUtils
 	{
 		/// <summary>
diff --git a/Engine/Source/Programs/DotNETCommon/DotNETUtilities/AssemblyUtils.cs b/Engine/Source/Programs/DotNETCommon/DotNETUtilities/AssemblyUtils.cs
--- a/Engine/Source/Programs/DotNETCommon/DotNETUtilities/AssemblyUtils.cs
+++ b/Engine/Source/Programs/DotNETCommon/DotNETUtilities/AssemblyUtils.cs
@@ -9,6 +9,7 @@ using System.Reflection;
 
 namespace Tools.DotNETCommon
 {
+	[Obsolete("Functionality in the Tools.DotNETCommon namespace is deprecated. Please reference the EpicGames.Common namespace and assembly instead.")]
 	public static class AssemblyUtils
 	{
         /// <summary>
diff --git a/Engine/Source/Programs/DotNETCommon/DotNETUtilities/BinaryArchiveReader.cs b/Engine/Source/Programs/DotNETCommon/DotNETUtilities/BinaryArchiveReader.cs
--- a/Engine/Source/Programs/DotNETCommon/DotNETUtilities/BinaryArchiveReader.cs
+++ b/Engine/Source/Programs/DotNETCommon/DotNETUtilities/BinaryArchiveReader.cs
@@ -13,6 +13,7 @@ namespace Tools.DotNETCommon
 	/// Reads data from a binary output stream. Similar to the NET Framework BinaryReader class, but supports fast serialization of object graphs and container types, and supports nullable objects.
 	/// Significantly faster than BinaryReader due to the expectation that the whole stream is in memory before deserialization.
 	/// </summary>
+	[Obsolete("Functionality in the Tools.DotNETCommon namespace is deprecated. Please reference the EpicGames.Common namespace and assembly instead.")]
 	public class BinaryArchiveReader : IDisposable
 	{
 		/// <summary>
diff --git a/Engine/Source/Programs/DotNETCommon/DotNETUtilities/BinaryArchiveWriter.cs b/Engine/Source/Programs/DotNETCommon/DotNETUtilities/BinaryArchiveWriter.cs
--- a/Engine/Source/Programs/DotNETCommon/DotNETUtilities/BinaryArchiveWriter.cs
+++ b/Engine/Source/Programs/DotNETCommon/DotNETUtilities/BinaryArchiveWriter.cs
@@ -13,6 +13,7 @@ namespace Tools.DotNETCommon
 	/// <summary>
 	/// Writes data to a binary output stream. Similar to the NET Framework BinaryWriter class, but supports fast serialization of object graphs and container types, and supports nullable objects.
 	/// </summary>
+	[Obsolete("Functionality in the Tools.DotNETCommon namespace is deprecated. Please reference the EpicGames.Common namespace and assembly instead.")]
 	public class BinaryArchiveWriter : IDisposable
 	{
 		/// <summary>
diff --git a/Engine/Source/Programs/DotNETCommon/DotNETUtilities/BinaryFormatterUtils.cs b/Engine/Source/Programs/DotNETCommon/DotNETUtilities/BinaryFormatterUtils.cs
--- a/Engine/Source/Programs/DotNETCommon/DotNETUtilities/BinaryFormatterUtils.cs
+++ b/Engine/Source/Programs/DotNETCommon/DotNETUtilities/BinaryFormatterUtils.cs
@@ -13,6 +13,7 @@ namespace Tools.DotNETCommon
 	/// <summary>
 	/// Utility functions for serializing using the BinaryFormatter
 	/// </summary>
+	[Obsolete("Functionality in the Tools.DotNETCommon namespace is deprecated. Please reference the EpicGames.Common namespace and assembly instead.")]
 	public static class BinaryFormatterUtils
 	{
 		/// <summary>
diff --git a/Engine/Source/Programs/DotNETCommon/DotNETUtilities/BinaryReaderExtensions.cs b/Engine/Source/Programs/DotNETCommon/DotNETUtilities/BinaryReaderExtensions.cs
--- a/Engine/Source/Programs/DotNETCommon/DotNETUtilities/BinaryReaderExtensions.cs
+++ b/Engine/Source/Programs/DotNETCommon/DotNETUtilities/BinaryReaderExtensions.cs
@@ -9,6 +9,7 @@ using System.Threading.Tasks;
 
 namespace Tools.DotNETCommon
 {
+	[Obsolete("Functionality in the Tools.DotNETCommon namespace is deprecated. Please reference the EpicGames.Common namespace and assembly instead.")]
 	public static class BinaryReaderExtensions
 	{
 		/// <summary>
diff --git a/Engine/Source/Programs/DotNETCommon/DotNETUtilities/BinaryWriterExtensions.cs b/Engine/Source/Programs/DotNETCommon/DotNETUtilities/BinaryWriterExtensions.cs
--- a/Engine/Source/Programs/DotNETCommon/DotNETUtilities/BinaryWriterExtensions.cs
+++ b/Engine/Source/Programs/DotNETCommon/DotNETUtilities/BinaryWriterExtensions.cs
@@ -14,6 +14,7 @@ namespace Tools.DotNETCommon
 	/// derived objects will have a constructor taking a single BinaryReader argument (similar to how the system ISerializable
 	/// interface assumes a specific constructor)
 	/// </summary>
+	[Obsolete("Functionality in the Tools.DotNETCommon namespace is deprecated. Please reference the EpicGames.Common namespace and assembly instead.")]
 	public interface IBinarySerializable
 	{
 		void Write(BinaryWriter Writer);
@@ -22,6 +23,7 @@ namespace Tools.DotNETCommon
 	/// <summary>
 	/// Extension methods for BinaryWriter class
 	/// </summary>
+	[Obsolete("Functionality in the Tools.DotNETCommon namespace is deprecated. Please reference the EpicGames.Common namespace and assembly instead.")]
 	public static class BinaryWriterExtensions
 	{
 		/// <summary>
diff --git a/Engine/Source/Programs/DotNETCommon/DotNETUtilities/CommandLineArguments.cs b/Engine/Source/Programs/DotNETCommon/DotNETUtilities/CommandLineArguments.cs
--- a/Engine/Source/Programs/DotNETCommon/DotNETUtilities/CommandLineArguments.cs
+++ b/Engine/Source/Programs/DotNETCommon/DotNETUtilities/CommandLineArguments.cs
@@ -15,6 +15,7 @@ namespace Tools.DotNETCommon
 	/// <summary>
 	/// Helper class to visualize an argument list
 	/// </summary>
+	[Obsolete("Functionality in the Tools.DotNETCommon namespace is deprecated. Please reference the EpicGames.Common namespace and assembly instead.")]
 	class CommandLineArgumentListView
 	{
 		/// <summary>
@@ -36,6 +37,7 @@ namespace Tools.DotNETCommon
 	/// <summary>
 	/// Exception thrown for invalid command line arguments
 	/// </summary>
+	[Obsolete("Functionality in the Tools.DotNETCommon namespace is deprecated. Please reference the EpicGames.Common namespace and assembly instead.")]
 	public class CommandLineArgumentException : Exception
 	{
 		/// <summary>
diff --git a/Engine/Source/Programs/DotNETCommon/DotNETUtilities/CommandLineAttribute.cs b/Engine/Source/Programs/DotNETCommon/DotNETUtilities/CommandLineAttribute.cs
--- a/Engine/Source/Programs/DotNETCommon/DotNETUtilities/CommandLineAttribute.cs
+++ b/Engine/Source/Programs/DotNETCommon/DotNETUtilities/CommandLineAttribute.cs
@@ -12,6 +12,7 @@ namespace Tools.DotNETCommon
 	/// <summary>
 	/// Attribute to indicate the name of a command line argument
 	/// </summary>
+	[Obsolete("Functionality in the Tools.DotNETCommon namespace is deprecated. Please reference the EpicGames.Common namespace and assembly instead.")]
 	[AttributeUsage(AttributeTargets.Field, AllowMultiple = true)]
 	public class CommandLineAttribute : Attribute
 	{
diff --git a/Engine/Source/Programs/DotNETCommon/DotNETUtilities/ContentHash.cs b/Engine/Source/Programs/DotNETCommon/DotNETUtilities/ContentHash.cs
--- a/Engine/Source/Programs/DotNETCommon/DotNETUtilities/ContentHash.cs
+++ b/Engine/Source/Programs/DotNETCommon/DotNETUtilities/ContentHash.cs
@@ -13,6 +13,7 @@ namespace Tools.DotNETCommon
 	/// <summary>
 	/// Stores the hash value for a piece of content as a byte array, allowing it to be used as a dictionary key
 	/// </summary>
+	[Obsolete("Functionality in the Tools.DotNETCommon namespace is deprecated. Please reference the EpicGames.Common namespace and assembly instead.")]
 	public class ContentHash : IEquatable<ContentHash>
 	{
 		public const int LengthMD5 = 16;
diff --git a/Engine/Source/Programs/DotNETCommon/DotNETUtilities/CsProjectInfo.cs b/Engine/Source/Programs/DotNETCommon/DotNETUtilities/CsProjectInfo.cs
--- a/Engine/Source/Programs/DotNETCommon/DotNETUtilities/CsProjectInfo.cs
+++ b/Engine/Source/Programs/DotNETCommon/DotNETUtilities/CsProjectInfo.cs
@@ -14,6 +14,7 @@ namespace Tools.DotNETCommon
 	/// <summary>
 	/// Basic information from a preprocessed C# project file. Supports reading a project file, expanding simple conditions in it, parsing property values, assembly references and references to other projects.
 	/// </summary>
+	[Obsolete("Functionality in the Tools.DotNETCommon namespace is deprecated. Please reference the EpicGames.Common namespace and assembly instead.")]
 	public class CsProjectInfo
 	{
 		/// <summary>
diff --git a/Engine/Source/Programs/DotNETCommon/DotNETUtilities/DirectoryReference.cs b/Engine/Source/Programs/DotNETCommon/DotNETUtilities/DirectoryReference.cs
--- a/Engine/Source/Programs/DotNETCommon/DotNETUtilities/DirectoryReference.cs
+++ b/Engine/Source/Programs/DotNETCommon/DotNETUtilities/DirectoryReference.cs
@@ -13,6 +13,7 @@ namespace Tools.DotNETCommon
 	/// Representation of an absolute directory path. Allows fast hashing and comparisons.
 	/// </summary>
 	[Serializable]
+	[Obsolete("Functionality in the Tools.DotNETCommon namespace is deprecated. Please reference the EpicGames.Common namespace and assembly instead.")]
 	public class DirectoryReference : FileSystemReference, IEquatable<DirectoryReference>
 	{
 		/// <summary>
diff --git a/Engine/Source/Programs/DotNETCommon/DotNETUtilities/DirectoryUtils.cs b/Engine/Source/Programs/DotNETCommon/DotNETUtilities/DirectoryUtils.cs
--- a/Engine/Source/Programs/DotNETCommon/DotNETUtilities/DirectoryUtils.cs
+++ b/Engine/Source/Programs/DotNETCommon/DotNETUtilities/DirectoryUtils.cs
@@ -12,6 +12,7 @@ namespace Tools.DotNETCommon
 	/// <summary>
 	/// Utility functions for manipulating directories
 	/// </summary>
+	[Obsolete("Functionality in the Tools.DotNETCommon namespace is deprecated. Please reference the EpicGames.Common namespace and assembly instead.")]
 	public class DirectoryUtils
 	{
 		/// <summary>
diff --git a/Engine/Source/Programs/DotNETCommon/DotNETUtilities/DotNETUtilities.csproj b/Engine/Source/Programs/DotNETCommon/DotNETUtilities/DotNETUtilities.csproj
--- a/Engine/Source/Programs/DotNETCommon/DotNETUtilities/DotNETUtilities.csproj
+++ b/Engine/Source/Programs/DotNETCommon/DotNETUtilities/DotNETUtilities.csproj
@@ -1,4 +1,4 @@
-<Project Sdk="Microsoft.NET.Sdk">
+﻿<Project Sdk="Microsoft.NET.Sdk">
   <PropertyGroup>
     <TargetFramework>netcoreapp3.1</TargetFramework>
     <Configuration Condition=" '$(Configuration)' == '' ">Development</Configuration>
@@ -10,6 +10,7 @@
   </PropertyGroup>
   <PropertyGroup Condition=" '$(Configuration)|$(Platform)' == 'Debug|AnyCPU' ">
     <DocumentationFile></DocumentationFile>
+    <NoWarn>1701;1702;CS0618</NoWarn>
   </PropertyGroup>
   <PropertyGroup Condition="'$(Configuration)|$(Platform)' == 'Development|AnyCPU'">
     <DocumentationFile></DocumentationFile>
diff --git a/Engine/Source/Programs/DotNETCommon/DotNETUtilities/ExceptionUtils.cs b/Engine/Source/Programs/DotNETCommon/DotNETUtilities/ExceptionUtils.cs
--- a/Engine/Source/Programs/DotNETCommon/DotNETUtilities/ExceptionUtils.cs
+++ b/Engine/Source/Programs/DotNETCommon/DotNETUtilities/ExceptionUtils.cs
@@ -12,6 +12,7 @@ namespace Tools.DotNETCommon
 	/// <summary>
 	/// Methods for adding context information to exceptions
 	/// </summary>
+	[Obsolete("Functionality in the Tools.DotNETCommon namespace is deprecated. Please reference the EpicGames.Common namespace and assembly instead.")]
 	public static class ExceptionUtils
 	{
 		/// <summary>
diff --git a/Engine/Source/Programs/DotNETCommon/DotNETUtilities/FatalErrorException.cs b/Engine/Source/Programs/DotNETCommon/DotNETUtilities/FatalErrorException.cs
--- a/Engine/Source/Programs/DotNETCommon/DotNETUtilities/FatalErrorException.cs
+++ b/Engine/Source/Programs/DotNETCommon/DotNETUtilities/FatalErrorException.cs
@@ -11,6 +11,7 @@ namespace Tools.DotNETCommon
 	/// <summary>
 	/// Base class for user-facing fatal errors. These errors are shown to the user prior to termination, without a callstack, and may dictate the program exit code.
 	/// </summary>
+	[Obsolete("Functionality in the Tools.DotNETCommon namespace is deprecated. Please reference the EpicGames.Common namespace and assembly instead.")]
 	public class FatalErrorException : Exception
 	{
 		/// <summary>
diff --git a/Engine/Source/Programs/DotNETCommon/DotNETUtilities/FileFilter.cs b/Engine/Source/Programs/DotNETCommon/DotNETUtilities/FileFilter.cs
--- a/Engine/Source/Programs/DotNETCommon/DotNETUtilities/FileFilter.cs
+++ b/Engine/Source/Programs/DotNETCommon/DotNETUtilities/FileFilter.cs
@@ -14,6 +14,7 @@ namespace Tools.DotNETCommon
 	/// <summary>
 	/// Indicates whether files which match a pattern should be included or excluded
 	/// </summary>
+	[Obsolete("Functionality in the Tools.DotNETCommon namespace is deprecated. Please reference the EpicGames.Common namespace and assembly instead.")]
 	public enum FileFilterType
 	{
 		/// <summary>
@@ -32,6 +33,7 @@ namespace Tools.DotNETCommon
 	/// The rules are merged into a tree with nodes for each path fragments and rule numbers in each leaf, allowing tree traversal in an arbitrary order 
 	/// to how they are applied.
 	/// </summary>
+	[Obsolete("Functionality in the Tools.DotNETCommon namespace is deprecated. Please reference the EpicGames.Common namespace and assembly instead.")]
 	public class FileFilter
 	{
 		/// <summary>
diff --git a/Engine/Source/Programs/DotNETCommon/DotNETUtilities/FileFilterNode.cs b/Engine/Source/Programs/DotNETCommon/DotNETUtilities/FileFilterNode.cs
--- a/Engine/Source/Programs/DotNETCommon/DotNETUtilities/FileFilterNode.cs
+++ b/Engine/Source/Programs/DotNETCommon/DotNETUtilities/FileFilterNode.cs
@@ -11,6 +11,7 @@ namespace Tools.DotNETCommon
 	/// <summary>
 	/// Node within a filter tree. Each node matches a single path fragment - a folder or file name, with an include or exclude rule consisting of a sequence of nodes. 
 	/// </summary>
+	[Obsolete("Functionality in the Tools.DotNETCommon namespace is deprecated. Please reference the EpicGames.Common namespace and assembly instead.")]
 	class FileFilterNode
 	{
 		/// <summary>
diff --git a/Engine/Source/Programs/DotNETCommon/DotNETUtilities/FilePattern.cs b/Engine/Source/Programs/DotNETCommon/DotNETUtilities/FilePattern.cs
--- a/Engine/Source/Programs/DotNETCommon/DotNETUtilities/FilePattern.cs
+++ b/Engine/Source/Programs/DotNETCommon/DotNETUtilities/FilePattern.cs
@@ -11,6 +11,7 @@ using Tools.DotNETCommon;
 
 namespace Tools.DotNETCommon
 {
+	[Obsolete("Functionality in the Tools.DotNETCommon namespace is deprecated. Please reference the EpicGames.Common namespace and assembly instead.")]
 	public class FilePatternException : Exception
 	{
 		public FilePatternException(string Message)
@@ -32,6 +33,7 @@ namespace Tools.DotNETCommon
 	/// <summary>
 	/// Encapsulates a pattern containing the '?', '*', and '...' wildcards.
 	/// </summary>
+	[Obsolete("Functionality in the Tools.DotNETCommon namespace is deprecated. Please reference the EpicGames.Common namespace and assembly instead.")]
 	public class FilePattern
 	{
 		/// <summary>
diff --git a/Engine/Source/Programs/DotNETCommon/DotNETUtilities/FileReference.cs b/Engine/Source/Programs/DotNETCommon/DotNETUtilities/FileReference.cs
--- a/Engine/Source/Programs/DotNETCommon/DotNETUtilities/FileReference.cs
+++ b/Engine/Source/Programs/DotNETCommon/DotNETUtilities/FileReference.cs
@@ -12,6 +12,7 @@ namespace Tools.DotNETCommon
 	/// <summary>
 	/// Representation of an absolute file path. Allows fast hashing and comparisons.
 	/// </summary>
+	[Obsolete("Functionality in the Tools.DotNETCommon namespace is deprecated. Please reference the EpicGames.Common namespace and assembly instead.")]
 	[Serializable]
 	public class FileReference : FileSystemReference, IEquatable<FileReference>
 	{
diff --git a/Engine/Source/Programs/DotNETCommon/DotNETUtilities/FileSystemReference.cs b/Engine/Source/Programs/DotNETCommon/DotNETUtilities/FileSystemReference.cs
--- a/Engine/Source/Programs/DotNETCommon/DotNETUtilities/FileSystemReference.cs
+++ b/Engine/Source/Programs/DotNETCommon/DotNETUtilities/FileSystemReference.cs
@@ -12,6 +12,7 @@ namespace Tools.DotNETCommon
 	/// <summary>
 	/// Base class for file system objects (files or directories).
 	/// </summary>
+	[Obsolete("Functionality in the Tools.DotNETCommon namespace is deprecated. Please reference the EpicGames.Common namespace and assembly instead.")]
 	[Serializable]
 	public abstract class FileSystemReference
 	{
diff --git a/Engine/Source/Programs/DotNETCommon/DotNETUtilities/FileUtils.cs b/Engine/Source/Programs/DotNETCommon/DotNETUtilities/FileUtils.cs
--- a/Engine/Source/Programs/DotNETCommon/DotNETUtilities/FileUtils.cs
+++ b/Engine/Source/Programs/DotNETCommon/DotNETUtilities/FileUtils.cs
@@ -14,6 +14,7 @@ namespace Tools.DotNETCommon
 	/// <summary>
 	/// Exception used to represent caught file/directory exceptions.
 	/// </summary>
+	[Obsolete("Functionality in the Tools.DotNETCommon namespace is deprecated. Please reference the EpicGames.Common namespace and assembly instead.")]
 	class WrappedFileOrDirectoryException : Exception
 	{
 		/// <summary>
@@ -38,6 +39,7 @@ namespace Tools.DotNETCommon
 	/// <summary>
 	/// Utility functions for manipulating files. Where these methods have similar functionality to those in the NET Framework, they generally go the extra mile to produce concise, specific error messages where possible.
 	/// </summary>
+	[Obsolete("Functionality in the Tools.DotNETCommon namespace is deprecated. Please reference the EpicGames.Common namespace and assembly instead.")]
 	public static class FileUtils
 	{
 		/// <summary>
diff --git a/Engine/Source/Programs/DotNETCommon/DotNETUtilities/HelpUtils.cs b/Engine/Source/Programs/DotNETCommon/DotNETUtilities/HelpUtils.cs
--- a/Engine/Source/Programs/DotNETCommon/DotNETUtilities/HelpUtils.cs
+++ b/Engine/Source/Programs/DotNETCommon/DotNETUtilities/HelpUtils.cs
@@ -12,6 +12,7 @@ namespace Tools.DotNETCommon
 	/// <summary>
 	/// Utility functions for showing help for objects
 	/// </summary>
+	[Obsolete("Functionality in the Tools.DotNETCommon namespace is deprecated. Please reference the EpicGames.Common namespace and assembly instead.")]
 	public static class HelpUtils
 	{
 		/// <summary>
diff --git a/Engine/Source/Programs/DotNETCommon/DotNETUtilities/Json.cs b/Engine/Source/Programs/DotNETCommon/DotNETUtilities/Json.cs
--- a/Engine/Source/Programs/DotNETCommon/DotNETUtilities/Json.cs
+++ b/Engine/Source/Programs/DotNETCommon/DotNETUtilities/Json.cs
@@ -12,6 +12,7 @@ namespace Tools.DotNETCommon
 	/// <summary>
 	/// Options for serializing an object to JSON
 	/// </summary>
+	[Obsolete("Functionality in the Tools.DotNETCommon namespace is deprecated. Please reference the EpicGames.Common namespace and assembly instead.")]
 	public enum JsonSerializeOptions
 	{
 		/// <summary>
@@ -28,6 +29,7 @@ namespace Tools.DotNETCommon
 	/// <summary>
 	/// Wrapper class around fastJSON methods.
 	/// </summary>
+	[Obsolete("Functionality in the Tools.DotNETCommon namespace is deprecated. Please reference the EpicGames.Common namespace and assembly instead.")]
 	public static class Json
 	{
 		/// <summary>
diff --git a/Engine/Source/Programs/DotNETCommon/DotNETUtilities/JsonObject.cs b/Engine/Source/Programs/DotNETCommon/DotNETUtilities/JsonObject.cs
--- a/Engine/Source/Programs/DotNETCommon/DotNETUtilities/JsonObject.cs
+++ b/Engine/Source/Programs/DotNETCommon/DotNETUtilities/JsonObject.cs
@@ -1,4 +1,4 @@
-﻿// Copyright Epic Games, Inc. All Rights Reserved.
+// Copyright Epic Games, Inc. All Rights Reserved.
 
 using System;
 using System.Collections.Generic;
@@ -11,6 +11,7 @@ namespace Tools.DotNETCommon
 	/// <summary>
 	/// Exception thrown for errors parsing JSON files
 	/// </summary>
+	[Obsolete("Functionality in the Tools.DotNETCommon namespace is deprecated. Please reference the EpicGames.Common namespace and assembly instead.")]
 	public class JsonParseException : Exception
 	{
 		/// <summary>
@@ -27,6 +28,7 @@ namespace Tools.DotNETCommon
 	/// <summary>
 	/// Stores a JSON object in memory
 	/// </summary>
+	[Obsolete("Functionality in the Tools.DotNETCommon namespace is deprecated. Please reference the EpicGames.Common namespace and assembly instead.")]
 	public class JsonObject
 	{
 		Dictionary<string, object> RawObject;
diff --git a/Engine/Source/Programs/DotNETCommon/DotNETUtilities/JsonWriter.cs b/Engine/Source/Programs/DotNETCommon/DotNETUtilities/JsonWriter.cs
--- a/Engine/Source/Programs/DotNETCommon/DotNETUtilities/JsonWriter.cs
+++ b/Engine/Source/Programs/DotNETCommon/DotNETUtilities/JsonWriter.cs
@@ -12,6 +12,7 @@ namespace Tools.DotNETCommon
 	/// <summary>
 	/// Specifies how to format JSON output
 	/// </summary>
+	[Obsolete("Functionality in the Tools.DotNETCommon namespace is deprecated. Please reference the EpicGames.Common namespace and assembly instead.")]
 	public enum JsonWriterStyle
 	{
 		/// <summary>
@@ -28,6 +29,7 @@ namespace Tools.DotNETCommon
 	/// <summary>
 	/// Writer for JSON data, which indents the output text appropriately, and adds commas and newlines between fields
 	/// </summary>
+	[Obsolete("Functionality in the Tools.DotNETCommon namespace is deprecated. Please reference the EpicGames.Common namespace and assembly instead.")]
 	public class JsonWriter : IDisposable
 	{
 		TextWriter Writer;
diff --git a/Engine/Source/Programs/DotNETCommon/DotNETUtilities/Log.cs b/Engine/Source/Programs/DotNETCommon/DotNETUtilities/Log.cs
--- a/Engine/Source/Programs/DotNETCommon/DotNETUtilities/Log.cs
+++ b/Engine/Source/Programs/DotNETCommon/DotNETUtilities/Log.cs
@@ -14,6 +14,7 @@ namespace Tools.DotNETCommon
 	/// <summary>
 	/// Log Event Type
 	/// </summary>
+	[Obsolete("Functionality in the Tools.DotNETCommon namespace is deprecated. Please reference the EpicGames.Common namespace and assembly instead.")]
 	public enum LogEventType
 	{
 		/// <summary>
@@ -55,6 +56,7 @@ namespace Tools.DotNETCommon
 	/// <summary>
 	/// Options for formatting messages
 	/// </summary>
+	[Obsolete("Functionality in the Tools.DotNETCommon namespace is deprecated. Please reference the EpicGames.Common namespace and assembly instead.")]
 	[Flags]
 	public enum LogFormatOptions
 	{
@@ -86,6 +88,7 @@ namespace Tools.DotNETCommon
 	///   2. We have a verbose filter we would like to apply to all logs without having to have each listener filter individually, which would require our string formatting code to run every time.
 	///   3. We possibly want to ensure severity prefixes are logged, but Trace.WriteXXX does not allow any severity info to be passed down.
 	/// </summary>
+	[Obsolete("Functionality in the Tools.DotNETCommon namespace is deprecated. Please reference the EpicGames.Common namespace and assembly instead.")]
 	static public class Log
 	{
 		/// <summary>
diff --git a/Engine/Source/Programs/DotNETCommon/DotNETUtilities/LogIndentScope.cs b/Engine/Source/Programs/DotNETCommon/DotNETUtilities/LogIndentScope.cs
--- a/Engine/Source/Programs/DotNETCommon/DotNETUtilities/LogIndentScope.cs
+++ b/Engine/Source/Programs/DotNETCommon/DotNETUtilities/LogIndentScope.cs
@@ -11,6 +11,7 @@ namespace Tools.DotNETCommon
 	/// <summary>
 	/// Class to apply a log indent for the lifetime of an object 
 	/// </summary>
+	[Obsolete("Functionality in the Tools.DotNETCommon namespace is deprecated. Please reference the EpicGames.Common namespace and assembly instead.")]
 	public class LogIndentScope : IDisposable
 	{
 		/// <summary>
diff --git a/Engine/Source/Programs/DotNETCommon/DotNETUtilities/LogStatusScope.cs b/Engine/Source/Programs/DotNETCommon/DotNETUtilities/LogStatusScope.cs
--- a/Engine/Source/Programs/DotNETCommon/DotNETUtilities/LogStatusScope.cs
+++ b/Engine/Source/Programs/DotNETCommon/DotNETUtilities/LogStatusScope.cs
@@ -11,6 +11,7 @@ namespace Tools.DotNETCommon
 	/// <summary>
 	/// Writes a status message to the log, which can be updated with a progress indicator as a slow task is being performed.
 	/// </summary>
+	[Obsolete("Functionality in the Tools.DotNETCommon namespace is deprecated. Please reference the EpicGames.Common namespace and assembly instead.")]
 	public class LogStatusScope : IDisposable
 	{
 		/// <summary>
diff --git a/Engine/Source/Programs/DotNETCommon/DotNETUtilities/ManagedProcess.cs b/Engine/Source/Programs/DotNETCommon/DotNETUtilities/ManagedProcess.cs
--- a/Engine/Source/Programs/DotNETCommon/DotNETUtilities/ManagedProcess.cs
+++ b/Engine/Source/Programs/DotNETCommon/DotNETUtilities/ManagedProcess.cs
@@ -17,6 +17,7 @@ namespace Tools.DotNETCommon
 	/// <summary>
 	/// Tracks a set of processes, and destroys them when the object is disposed.
 	/// </summary>
+	[Obsolete("Functionality in the Tools.DotNETCommon namespace is deprecated. Please reference the EpicGames.Common namespace and assembly instead.")]
 	public class ManagedProcessGroup : IDisposable
 	{
 		const UInt32 JOB_OBJECT_LIMIT_KILL_ON_JOB_CLOSE = 0x00002000;
@@ -135,6 +136,7 @@ namespace Tools.DotNETCommon
 	/// Uses job objects to ensure that the process will be terminated automatically by the O/S if the current process is terminated, and polls pipe reads to avoid blocking on unmanaged code
 	/// if the calling thread is terminated. Currently only implemented for Windows; makes heavy use of P/Invoke.
 	/// </summary>
+	[Obsolete("Functionality in the Tools.DotNETCommon namespace is deprecated. Please reference the EpicGames.Common namespace and assembly instead.")]
 	public sealed class ManagedProcess : IDisposable
 	{
 		[StructLayout(LayoutKind.Sequential)]
diff --git a/Engine/Source/Programs/DotNETCommon/DotNETUtilities/ProcessUtils.cs b/Engine/Source/Programs/DotNETCommon/DotNETUtilities/ProcessUtils.cs
--- a/Engine/Source/Programs/DotNETCommon/DotNETUtilities/ProcessUtils.cs
+++ b/Engine/Source/Programs/DotNETCommon/DotNETUtilities/ProcessUtils.cs
@@ -15,6 +15,7 @@ namespace Tools.DotNETCommon
 	/// <summary>
 	/// Includes utilities for managing processes
 	/// </summary>
+	[Obsolete("Functionality in the Tools.DotNETCommon namespace is deprecated. Please reference the EpicGames.Common namespace and assembly instead.")]
 	public static class ProcessUtils
 	{
 		[StructLayout(LayoutKind.Sequential)]
diff --git a/Engine/Source/Programs/DotNETCommon/DotNETUtilities/ReadOnlyHashSet.cs b/Engine/Source/Programs/DotNETCommon/DotNETUtilities/ReadOnlyHashSet.cs
--- a/Engine/Source/Programs/DotNETCommon/DotNETUtilities/ReadOnlyHashSet.cs
+++ b/Engine/Source/Programs/DotNETCommon/DotNETUtilities/ReadOnlyHashSet.cs
@@ -13,6 +13,7 @@ namespace Tools.DotNETCommon
 	/// Wrapper around the HashSet container that only allows read operations
 	/// </summary>
 	/// <typeparam name="T">Type of element for the hashset</typeparam>
+	[Obsolete("Functionality in the Tools.DotNETCommon namespace is deprecated. Please reference the EpicGames.Common namespace and assembly instead.")]
 	public class ReadOnlyHashSet<T> : IReadOnlyCollection<T>
 	{
 		/// <summary>
diff --git a/Engine/Source/Programs/DotNETCommon/DotNETUtilities/StringUtils.cs b/Engine/Source/Programs/DotNETCommon/DotNETUtilities/StringUtils.cs
--- a/Engine/Source/Programs/DotNETCommon/DotNETUtilities/StringUtils.cs
+++ b/Engine/Source/Programs/DotNETCommon/DotNETUtilities/StringUtils.cs
@@ -8,6 +8,7 @@ using System.Threading.Tasks;
 
 namespace Tools.DotNETCommon
 {
+	[Obsolete("Functionality in the Tools.DotNETCommon namespace is deprecated. Please reference the EpicGames.Common namespace and assembly instead.")]
 	public static class StringUtils
 	{
 		/// <summary>
diff --git a/Engine/Source/Programs/DotNETCommon/DotNETUtilities/ThreadPoolWorkQueue.cs b/Engine/Source/Programs/DotNETCommon/DotNETUtilities/ThreadPoolWorkQueue.cs
--- a/Engine/Source/Programs/DotNETCommon/DotNETUtilities/ThreadPoolWorkQueue.cs
+++ b/Engine/Source/Programs/DotNETCommon/DotNETUtilities/ThreadPoolWorkQueue.cs
@@ -12,6 +12,7 @@ namespace Tools.DotNETCommon
 	/// <summary>
 	/// Allows queuing a large number of tasks to a thread pool and waiting for them all to complete.
 	/// </summary>
+	[Obsolete("Functionality in the Tools.DotNETCommon namespace is deprecated. Please reference the EpicGames.Common namespace and assembly instead.")]
 	public class ThreadPoolWorkQueue : IDisposable
 	{
 		/// <summary>
diff --git a/Engine/Source/Programs/DotNETCommon/DotNETUtilities/TraceSpan.cs b/Engine/Source/Programs/DotNETCommon/DotNETUtilities/TraceSpan.cs
--- a/Engine/Source/Programs/DotNETCommon/DotNETUtilities/TraceSpan.cs
+++ b/Engine/Source/Programs/DotNETCommon/DotNETUtilities/TraceSpan.cs
@@ -15,6 +15,7 @@ namespace Tools.DotNETCommon
 	/// <summary>
 	/// Interface for a trace span
 	/// </summary>
+	[Obsolete("Functionality in the Tools.DotNETCommon namespace is deprecated. Please reference the EpicGames.Common namespace and assembly instead.")]
 	public interface ITraceSpan : IDisposable
 	{
 		/// <summary>
@@ -28,6 +29,7 @@ namespace Tools.DotNETCommon
 	/// <summary>
 	/// Methods for creating ITraceScope instances
 	/// </summary>
+	[Obsolete("Functionality in the Tools.DotNETCommon namespace is deprecated. Please reference the EpicGames.Common namespace and assembly instead.")]
 	public static class TraceSpan
 	{
 		/// <summary>
diff --git a/Engine/Source/Programs/DotNETCommon/DotNETUtilities/Win32ExceptionWithCode.cs b/Engine/Source/Programs/DotNETCommon/DotNETUtilities/Win32ExceptionWithCode.cs
--- a/Engine/Source/Programs/DotNETCommon/DotNETUtilities/Win32ExceptionWithCode.cs
+++ b/Engine/Source/Programs/DotNETCommon/DotNETUtilities/Win32ExceptionWithCode.cs
@@ -12,6 +12,7 @@ namespace Tools.DotNETCommon
 	/// <summary>
 	/// Wrapper for Win32Exception which includes the error code in the exception message
 	/// </summary>
+	[Obsolete("Functionality in the Tools.DotNETCommon namespace is deprecated. Please reference the EpicGames.Common namespace and assembly instead.")]
 	class Win32ExceptionWithCode : Win32Exception
 	{
 		/// <summary>
diff --git a/Engine/Source/Programs/Shared/EpicGames.Core/EpicGames.Core.csproj b/Engine/Source/Programs/Shared/EpicGames.Core/EpicGames.Core.csproj
--- a/Engine/Source/Programs/Shared/EpicGames.Core/EpicGames.Core.csproj
+++ b/Engine/Source/Programs/Shared/EpicGames.Core/EpicGames.Core.csproj
@@ -3,7 +3,7 @@
     <TargetFramework>netcoreapp3.1</TargetFramework>
     <Configuration Condition=" '$(Configuration)' == '' ">Development</Configuration>
     <OutputType>Library</OutputType>
-    <RootNamespace>EpicGames.Common</RootNamespace>
+    <RootNamespace>EpicGames.Core</RootNamespace>
     <GenerateAssemblyInfo>false</GenerateAssemblyInfo>
     <GenerateTargetFrameworkAttribute>false</GenerateTargetFrameworkAttribute>
     <Configurations>Debug;Release;Development</Configurations>
diff --git a/Engine/Source/Programs/Shared/EpicGames.Perforce/EpicGames.Perforce.csproj b/Engine/Source/Programs/Shared/EpicGames.Perforce/EpicGames.Perforce.csproj
--- a/Engine/Source/Programs/Shared/EpicGames.Perforce/EpicGames.Perforce.csproj
+++ b/Engine/Source/Programs/Shared/EpicGames.Perforce/EpicGames.Perforce.csproj
@@ -3,7 +3,7 @@
     <TargetFramework>netcoreapp3.1</TargetFramework>
     <Configuration Condition=" '$(Configuration)' == '' ">Development</Configuration>
     <OutputType>Library</OutputType>
-    <RootNamespace>EpicGames.Common</RootNamespace>
+    <RootNamespace>EpicGames.Perforce</RootNamespace>
     <GenerateAssemblyInfo>false</GenerateAssemblyInfo>
     <GenerateTargetFrameworkAttribute>false</GenerateTargetFrameworkAttribute>
     <Configurations>Debug;Release;Development</Configurations>
