- Fixed Android AutoSDK version from not showing up as valid in Turnkey

diff --git a/Engine/Source/Programs/DotNETCommon/BuildUtilities/UEBuildPlatformSDK.cs b/Engine/Source/Programs/DotNETCommon/BuildUtilities/UEBuildPlatformSDK.cs
--- a/Engine/Source/Programs/DotNETCommon/BuildUtilities/UEBuildPlatformSDK.cs
+++ b/Engine/Source/Programs/DotNETCommon/BuildUtilities/UEBuildPlatformSDK.cs
@@ -271,6 +271,13 @@ namespace EpicGames.Core
 				return false;
 			}
 
+			// AutoSDK must match the desired version exactly, since that is the only one we will use
+			if (bForAutoSDK)
+			{
+				// if integer version checking failed, then we can detect valid autosdk if the version matches the autosdk directory by name
+				return string.Compare(Version, GetAutoSDKDirectoryForMainVersion(), true) == 0;
+			}
+
 			// convert it to an integer
 			UInt64 IntVersion;
 			if (!TryConvertVersionToInt(Version, out IntVersion))
@@ -289,15 +296,6 @@ namespace EpicGames.Core
 			{
 				return true;
 			}
-			else
-			{
-				// AutoSDK must match the desired version exactly, since that is the only one we will use
-				if (bForAutoSDK)
-				{
-					// if integer version checking failed, then we can detect valid autosdk if the version matches the autosdk directory by name
-					return string.Compare(Version, GetAutoSDKDirectoryForMainVersion(), true) == 0;
-				}
-			}
 
 			// get numeric range
 			UInt64 MinVersion, MaxVersion;
