Move AutoSDK script version to Version.txt file inside the AutoSDK directory itself
 - Previously, any changes made to AutoSDKs had to be submitted along with a bump in the version string returned by GetRequiredScriptVersionString(). Since this string is in UBT, the string has to be bumped in every branch where we need the AutoSDK changes to get applied.
 - Now, the script version string is stored in Version.txt, meaning there's only one place the version needs to be bumped to force all UBT instances in all branches to reapply an SDK.
 - Made GetRequiredScriptVersionString() private and removed platform implementations.

diff --git a/Engine/Source/Programs/DotNETCommon/BuildUtilities/UEBuildPlatformSDK.cs b/Engine/Source/Programs/DotNETCommon/BuildUtilities/UEBuildPlatformSDK.cs
--- a/Engine/Source/Programs/DotNETCommon/BuildUtilities/UEBuildPlatformSDK.cs
+++ b/Engine/Source/Programs/DotNETCommon/BuildUtilities/UEBuildPlatformSDK.cs
@@ -453,6 +453,11 @@ namespace EpicGames.Core
 		/// </summary>
 		protected const string LastRunScriptVersionManifest = "CurrentlyInstalled.Version.txt";
 
+		/// <summary>
+		/// Filename of the script version file in each AutoSDK directory. Changing the contents of this file will force reinstallation of an AutoSDK.
+		/// </summary>
+		private const string ScriptVersionFilename = "Version.txt";
+
 		/// <summary>
 		/// Name of the file that holds environment variables of current SDK
 		/// </summary>
@@ -524,12 +529,22 @@ namespace EpicGames.Core
 
 		/// <summary>
 		/// Gets the version number of the SDK setup script itself.  The version in the base should ALWAYS be the master revision from the last refactor.
-		/// If you need to force a rebuild for a given platform, override this for the given platform.
+		/// If you need to force a rebuild for a given platform, modify the version file.
 		/// </summary>
 		/// <returns>Setup script version</returns>
-		protected virtual String GetRequiredScriptVersionString()
+		private string GetRequiredScriptVersionString()
 		{
-			return "3.0";
+			const string UnspecifiedVersion = "UnspecifiedScriptVersion";
+
+			string VersionFilename = Path.Combine(GetPathToPlatformAutoSDKs(), GetAutoSDKDirectoryForMainVersion(), ScriptVersionFilename);
+			if (File.Exists(VersionFilename))
+			{
+				return File.ReadAllLines(VersionFilename).FirstOrDefault() ?? UnspecifiedVersion;
+			}
+			else
+			{
+				return UnspecifiedVersion;
+			}
 		}
 
 		/// <summary>
diff --git a/Engine/Source/Programs/UnrealBuildTool/Platform/Android/AndroidPlatformSDK.cs b/Engine/Source/Programs/UnrealBuildTool/Platform/Android/AndroidPlatformSDK.cs
--- a/Engine/Source/Programs/UnrealBuildTool/Platform/Android/AndroidPlatformSDK.cs
+++ b/Engine/Source/Programs/UnrealBuildTool/Platform/Android/AndroidPlatformSDK.cs
@@ -143,11 +143,6 @@ namespace UnrealBuildTool
 			return true;
 		}
 
-		protected override String GetRequiredScriptVersionString()
-		{
-			return "3.6";
-		}
-
 		// prefer auto sdk on android as correct 'manual' sdk detection isn't great at the moment.
 		protected override bool PreferAutoSDK()
 		{
diff --git a/Engine/Source/Programs/UnrealBuildTool/Platform/Linux/LinuxPlatformSDK.cs b/Engine/Source/Programs/UnrealBuildTool/Platform/Linux/LinuxPlatformSDK.cs
--- a/Engine/Source/Programs/UnrealBuildTool/Platform/Linux/LinuxPlatformSDK.cs
+++ b/Engine/Source/Programs/UnrealBuildTool/Platform/Linux/LinuxPlatformSDK.cs
@@ -141,11 +141,6 @@ namespace UnrealBuildTool
 			return null;
 		}
 
-		protected override String GetRequiredScriptVersionString()
-		{
-			return "3.0";
-		}
-
 		protected override bool PreferAutoSDK()
 		{
 			// having LINUX_ROOT set (for legacy reasons or for convenience of cross-compiling certain third party libs) should not make UBT skip AutoSDKs
diff --git a/Engine/Source/Programs/UnrealBuildTool/Platform/Lumin/LuminPlatformSDK.cs b/Engine/Source/Programs/UnrealBuildTool/Platform/Lumin/LuminPlatformSDK.cs
--- a/Engine/Source/Programs/UnrealBuildTool/Platform/Lumin/LuminPlatformSDK.cs
+++ b/Engine/Source/Programs/UnrealBuildTool/Platform/Lumin/LuminPlatformSDK.cs
@@ -90,11 +90,6 @@ namespace UnrealBuildTool
 			return true;
 		}
 
-		protected override String GetRequiredScriptVersionString()
-		{
-			return "Lumin_15";
-		}
-
 		// prefer auto sdk on android as correct 'manual' sdk detection isn't great at the moment.
 		protected override bool PreferAutoSDK()
 		{
