Move FileItem and DirectoryItem into BuildUtilities

diff --git a/Engine/Source/Programs/DotNETCommon/BuildUtilities/BuildUtilities.csproj b/Engine/Source/Programs/DotNETCommon/BuildUtilities/BuildUtilities.csproj
--- a/Engine/Source/Programs/DotNETCommon/BuildUtilities/BuildUtilities.csproj
+++ b/Engine/Source/Programs/DotNETCommon/BuildUtilities/BuildUtilities.csproj
@@ -7,6 +7,7 @@
     <GenerateAssemblyInfo>false</GenerateAssemblyInfo>
     <GenerateTargetFrameworkAttribute>false</GenerateTargetFrameworkAttribute>
     <Configurations>Debug;Release;Development</Configurations>
+    <Nullable>enable</Nullable>
   </PropertyGroup>
   <PropertyGroup Condition=" '$(Configuration)|$(Platform)' == 'Debug|AnyCPU' ">
     <DocumentationFile></DocumentationFile>
diff --git a/Engine/Source/Programs/DotNETCommon/BuildUtilities/System/DirectoryItem.cs b/Engine/Source/Programs/DotNETCommon/BuildUtilities/System/DirectoryItem.cs
new file mode 100644
--- /dev/null
+++ b/Engine/Source/Programs/DotNETCommon/BuildUtilities/System/DirectoryItem.cs
@@ -0,0 +1,359 @@
+// Copyright Epic Games, Inc. All Rights Reserved.
+
+using System;
+using System.Collections.Concurrent;
+using System.Collections.Generic;
+using System.Diagnostics.CodeAnalysis;
+using System.IO;
+using System.Linq;
+using System.Text;
+using System.Threading.Tasks;
+using EpicGames.Core;
+
+namespace UnrealBuildBase
+{
+	/// <summary>
+	/// Stores the state of a directory. May or may not exist.
+	/// </summary>
+	public class DirectoryItem
+	{
+		/// <summary>
+		/// Full path to the directory on disk
+		/// </summary>
+		public readonly DirectoryReference Location;
+
+		/// <summary>
+		/// Cached value for whether the directory exists
+		/// </summary>
+		DirectoryInfo Info;
+
+		/// <summary>
+		/// Cached map of name to subdirectory item
+		/// </summary>
+		Dictionary<string, DirectoryItem>? Directories;
+
+		/// <summary>
+		/// Cached map of name to file
+		/// </summary>
+		Dictionary<string, FileItem>? Files;
+
+		/// <summary>
+		/// Global map of location to item
+		/// </summary>
+		static ConcurrentDictionary<DirectoryReference, DirectoryItem> LocationToItem = new ConcurrentDictionary<DirectoryReference, DirectoryItem>();
+
+		/// <summary>
+		/// Constructor
+		/// </summary>
+		/// <param name="Location">Path to this directory</param>
+		/// <param name="Info">Information about this directory</param>
+		private DirectoryItem(DirectoryReference Location, DirectoryInfo Info)
+		{
+			this.Location = Location;
+			this.Info = Info;
+		}
+
+		/// <summary>
+		/// The name of this directory
+		/// </summary>
+		public string Name
+		{
+			get { return Info.Name; }
+		}
+
+		/// <summary>
+		/// The full name of this directory
+		/// </summary>
+		public string FullName
+		{
+			get { return Location.FullName; }
+		}
+
+		/// <summary>
+		/// Whether the directory exists or not
+		/// </summary>
+		public bool Exists
+		{
+			get { return Info.Exists; }
+		}
+
+		/// <summary>
+		/// The last write time of the file.
+		/// </summary>
+		public DateTime LastWriteTimeUtc
+		{
+			get { return Info.LastWriteTimeUtc; }
+		}
+
+		/// <summary>
+		/// Gets the parent directory item
+		/// </summary>
+		public DirectoryItem? GetParentDirectoryItem()
+		{
+			if(Info.Parent == null)
+			{
+				return null;
+			}
+			else
+			{
+				return GetItemByDirectoryInfo(Info.Parent);
+			}
+		}
+
+		/// <summary>
+		/// Gets a new directory item by combining the existing directory item with the given path fragments
+		/// </summary>
+		/// <param name="BaseDirectory">Base directory to append path fragments to</param>
+		/// <param name="Fragments">The path fragments to append</param>
+		/// <returns>Directory item corresponding to the combined path</returns>
+		public static DirectoryItem Combine(DirectoryItem BaseDirectory, params string[] Fragments)
+		{
+			return DirectoryItem.GetItemByDirectoryReference(DirectoryReference.Combine(BaseDirectory.Location, Fragments));
+		}
+
+		/// <summary>
+		/// Finds or creates a directory item from its location
+		/// </summary>
+		/// <param name="Location">Path to the directory</param>
+		/// <returns>The directory item for this location</returns>
+		public static DirectoryItem GetItemByPath(string Location)
+		{
+			return GetItemByDirectoryReference(new DirectoryReference(Location));
+		}
+
+		/// <summary>
+		/// Finds or creates a directory item from its location
+		/// </summary>
+		/// <param name="Location">Path to the directory</param>
+		/// <returns>The directory item for this location</returns>
+		public static DirectoryItem GetItemByDirectoryReference(DirectoryReference Location)
+		{
+			DirectoryItem? Result;
+			if(!LocationToItem.TryGetValue(Location, out Result))
+			{
+				DirectoryItem NewItem = new DirectoryItem(Location, new DirectoryInfo(Location.FullName));
+				if(LocationToItem.TryAdd(Location, NewItem))
+				{
+					Result = NewItem;
+				}
+				else
+				{
+					Result = LocationToItem[Location];
+				}
+			}
+			return Result;
+		}
+
+		/// <summary>
+		/// Finds or creates a directory item from a DirectoryInfo object
+		/// </summary>
+		/// <param name="Info">Path to the directory</param>
+		/// <returns>The directory item for this location</returns>
+		public static DirectoryItem GetItemByDirectoryInfo(DirectoryInfo Info)
+		{
+			DirectoryReference Location = new DirectoryReference(Info);
+
+			DirectoryItem? Result;
+			if(!LocationToItem.TryGetValue(Location, out Result))
+			{
+				DirectoryItem NewItem = new DirectoryItem(Location, Info);
+				if(LocationToItem.TryAdd(Location, NewItem))
+				{
+					Result = NewItem;
+				}
+				else
+				{
+					Result = LocationToItem[Location];
+				}
+			}
+			return Result;
+		}
+
+		/// <summary>
+		/// Reset the contents of the directory and allow them to be fetched again
+		/// </summary>
+		public void ResetCachedInfo()
+		{
+			Info = new DirectoryInfo(Info.FullName);
+
+			Dictionary<string, DirectoryItem>? PrevDirectories = Directories;
+			if(PrevDirectories != null)
+			{
+				foreach(DirectoryItem SubDirectory in PrevDirectories.Values)
+				{
+					SubDirectory.ResetCachedInfo();
+				}
+				Directories = null;
+			}
+
+			Dictionary<string, FileItem>? PrevFiles = Files;
+			if(PrevFiles != null)
+			{
+				foreach(FileItem File in PrevFiles.Values)
+				{
+					File.ResetCachedInfo();
+				}
+				Files = null;
+			}
+		}
+
+		/// <summary>
+		/// Resets all cached directory info. Significantly reduces performance; do not use unless strictly necessary.
+		/// </summary>
+		public static void ResetAllCachedInfo_SLOW()
+		{
+			foreach(DirectoryItem Item in LocationToItem.Values)
+			{
+				Item.Info = new DirectoryInfo(Item.Info.FullName);
+				Item.Directories = null;
+				Item.Files = null;
+			}
+			FileItem.ResetAllCachedInfo_SLOW();
+		}
+
+		/// <summary>
+		/// Caches the subdirectories of this directories
+		/// </summary>
+		public void CacheDirectories()
+		{
+			if(Directories == null)
+			{
+				Dictionary<string, DirectoryItem> NewDirectories = new Dictionary<string, DirectoryItem>(DirectoryReference.Comparer);
+				if(Info.Exists)
+				{
+					foreach(DirectoryInfo SubDirectoryInfo in Info.EnumerateDirectories())
+					{
+						if(SubDirectoryInfo.Name.Length == 1 && SubDirectoryInfo.Name[0] == '.')
+						{
+							continue;
+						}
+						else if(SubDirectoryInfo.Name.Length == 2 && SubDirectoryInfo.Name[0] == '.' && SubDirectoryInfo.Name[1] == '.')
+						{
+							continue;
+						}
+						else
+						{
+							NewDirectories[SubDirectoryInfo.Name] = DirectoryItem.GetItemByDirectoryInfo(SubDirectoryInfo);
+						}
+					}
+				}
+				Directories = NewDirectories;
+			}
+		}
+
+		/// <summary>
+		/// Enumerates all the subdirectories
+		/// </summary>
+		/// <returns>Sequence of subdirectory items</returns>
+		public IEnumerable<DirectoryItem> EnumerateDirectories()
+		{
+			CacheDirectories();
+			return Directories!.Values;
+		}
+
+		/// <summary>
+		/// Attempts to get a sub-directory by name
+		/// </summary>
+		/// <param name="Name">Name of the directory</param>
+		/// <param name="OutDirectory">If successful receives the matching directory item with this name</param>
+		/// <returns>True if the file exists, false otherwise</returns>
+		public bool TryGetDirectory(string Name, [NotNullWhen(true)] out DirectoryItem? OutDirectory)
+		{
+			if(Name.Length > 0 && Name[0] == '.')
+			{
+				if(Name.Length == 1)
+				{
+					OutDirectory = this;
+					return true;
+				}
+				else if(Name.Length == 2 && Name[1] == '.')
+				{
+					OutDirectory = GetParentDirectoryItem();
+					return OutDirectory != null;
+				}
+			}
+
+			CacheDirectories();
+			return Directories!.TryGetValue(Name, out OutDirectory);
+		}
+
+		/// <summary>
+		/// Caches the files in this directory
+		/// </summary>
+		public void CacheFiles()
+		{
+			if(Files == null)
+			{
+				Dictionary<string, FileItem> NewFiles = new Dictionary<string, FileItem>(FileReference.Comparer);
+				if(Info.Exists)
+				{
+					foreach(FileInfo FileInfo in Info.EnumerateFiles())
+					{
+						FileItem FileItem = FileItem.GetItemByFileInfo(FileInfo);
+						FileItem.UpdateCachedDirectory(this);
+						NewFiles[FileInfo.Name] = FileItem;
+					}
+				}
+				Files = NewFiles;
+			}
+		}
+
+		/// <summary>
+		/// Enumerates all the files
+		/// </summary>
+		/// <returns>Sequence of FileItems</returns>
+		public IEnumerable<FileItem> EnumerateFiles()
+		{
+			CacheFiles();
+			return Files!.Values;
+		}
+
+		/// <summary>
+		/// Attempts to get a file from this directory by name. Unlike creating a file item and checking whether it exists, this will
+		/// not create a permanent FileItem object if it does not exist.
+		/// </summary>
+		/// <param name="Name">Name of the file</param>
+		/// <param name="OutFile">If successful receives the matching file item with this name</param>
+		/// <returns>True if the file exists, false otherwise</returns>
+		public bool TryGetFile(string Name, [NotNullWhen(true)] out FileItem? OutFile)
+		{
+			CacheFiles();
+			return Files!.TryGetValue(Name, out OutFile);
+		}
+
+		/// <summary>
+		/// Formats this object as a string for debugging
+		/// </summary>
+		/// <returns>Location of the directory</returns>
+		public override string ToString()
+		{
+			return Location.FullName;
+		}
+	}
+
+	/// <summary>
+	/// Helper functions for serialization
+	/// </summary>
+	public static class DirectoryItemExtensionMethods
+	{
+		/// <summary>
+		/// Read a directory item from a binary archive
+		/// </summary>
+		/// <param name="Reader">Reader to serialize data from</param>
+		/// <returns>Instance of the serialized directory item</returns>
+		public static DirectoryItem ReadDirectoryItem(this BinaryArchiveReader Reader)
+		{
+			return Reader.ReadObjectReference<DirectoryItem>(() => DirectoryItem.GetItemByDirectoryReference(Reader.ReadDirectoryReferenceNotNull()));
+		}
+
+		/// <summary>
+		/// Write a directory item to a binary archive
+		/// </summary>
+		/// <param name="Writer">Writer to serialize data to</param>
+		/// <param name="DirectoryItem">Directory item to write</param>
+		public static void WriteDirectoryItem(this BinaryArchiveWriter Writer, DirectoryItem DirectoryItem)
+		{
+			Writer.WriteObjectReference<DirectoryItem>(DirectoryItem, () => Writer.WriteDirectoryReference(DirectoryItem.Location));
+		}
+	}
+}
diff --git a/Engine/Source/Programs/DotNETCommon/BuildUtilities/System/FileItem.cs b/Engine/Source/Programs/DotNETCommon/BuildUtilities/System/FileItem.cs
new file mode 100644
--- /dev/null
+++ b/Engine/Source/Programs/DotNETCommon/BuildUtilities/System/FileItem.cs
@@ -0,0 +1,347 @@
+// Copyright Epic Games, Inc. All Rights Reserved.
+
+using System;
+using System.Collections.Generic;
+using System.Text;
+using System.IO;
+using System.Diagnostics;
+using System.Threading;
+using System.Runtime.Serialization;
+using EpicGames.Core;
+using System.Collections.Concurrent;
+
+namespace UnrealBuildBase
+{
+	/// <summary>
+	/// Represents a file on disk that is used as an input or output of a build action. FileItem instances are unique for a given path. Use FileItem.GetItemByFileReference 
+	/// to get the FileItem for a specific path.
+	/// </summary>
+	public class FileItem
+	{
+		/// <summary>
+		/// The directory containing this file
+		/// </summary>
+		DirectoryItem? CachedDirectory;
+
+		/// <summary>
+		/// Location of this file
+		/// </summary>
+		public readonly FileReference Location;
+
+		/// <summary>
+		/// The information about the file.
+		/// </summary>
+		FileInfo Info;
+
+		/// <summary>
+		/// A case-insensitive dictionary that's used to map each unique file name to a single FileItem object.
+		/// </summary>
+		static ConcurrentDictionary<FileReference, FileItem> UniqueSourceFileMap = new ConcurrentDictionary<FileReference, FileItem>();
+
+		/// <summary>
+		/// Constructor
+		/// </summary>
+		/// <param name="Location">Location of the file</param>
+		/// <param name="Info">File info</param>
+		private FileItem(FileReference Location, FileInfo Info)
+		{
+			this.Location = Location;
+			this.Info = Info;
+		}
+
+		/// <summary>
+		/// Name of this file
+		/// </summary>
+		public string Name
+		{
+			get { return Info.Name; }
+		}
+
+		/// <summary>
+		/// Full name of this file
+		/// </summary>
+		public string FullName
+		{
+			get { return Location.FullName; }
+		}
+
+		/// <summary>
+		/// Accessor for the absolute path to the file
+		/// </summary>
+		public string AbsolutePath
+		{
+			get { return Location.FullName; }
+		}
+
+		/// <summary>
+		/// Gets the directory that this file is in
+		/// </summary>
+		public DirectoryItem Directory
+		{
+			get
+			{
+				if(CachedDirectory == null)
+				{
+					CachedDirectory = DirectoryItem.GetItemByDirectoryReference(Location.Directory);
+				}
+				return CachedDirectory;
+			}
+		}
+
+		/// <summary>
+		/// Whether the file exists.
+		/// </summary>
+		public bool Exists
+		{
+			get { return Info.Exists; }
+		}
+
+		/// <summary>
+		/// Size of the file if it exists, otherwise -1
+		/// </summary>
+		public long Length
+		{
+			get { return Info.Length; }
+		}
+
+		/// <summary>
+		/// The attributes for this file
+		/// </summary>
+		public FileAttributes Attributes
+		{
+			get { return Info.Attributes; }
+		}
+
+		/// <summary>
+		/// The last write time of the file.
+		/// </summary>
+		public DateTime LastWriteTimeUtc
+		{
+			get { return Info.LastWriteTimeUtc; }
+		}
+
+		/// <summary>
+		/// Determines if the file has the given extension
+		/// </summary>
+		/// <param name="Extension">The extension to check for</param>
+		/// <returns>True if the file has the given extension, false otherwise</returns>
+		public bool HasExtension(string Extension)
+		{
+			return Location.HasExtension(Extension);
+		}
+
+		/// <summary>
+		/// Gets the directory containing this file
+		/// </summary>
+		/// <returns>DirectoryItem for the directory containing this file</returns>
+		public DirectoryItem GetDirectoryItem()
+		{
+			return Directory;
+		}
+
+		/// <summary>
+		/// Updates the cached directory for this file. Used by DirectoryItem when enumerating files, to avoid having to look this up later.
+		/// </summary>
+		/// <param name="Directory">The directory that this file is in</param>
+		public void UpdateCachedDirectory(DirectoryItem Directory)
+		{
+			Debug.Assert(Directory.Location == Location.Directory);
+			CachedDirectory = Directory;
+		}
+
+		/// <summary>
+		/// Gets a FileItem corresponding to the given path
+		/// </summary>
+		/// <param name="FilePath">Path for the FileItem</param>
+		/// <returns>The FileItem that represents the given file path.</returns>
+		public static FileItem GetItemByPath(string FilePath)
+		{
+			return GetItemByFileReference(new FileReference(FilePath));
+		}
+
+		/// <summary>
+		/// Gets a FileItem for a given path
+		/// </summary>
+		/// <param name="Info">Information about the file</param>
+		/// <returns>The FileItem that represents the given a full file path.</returns>
+		public static FileItem GetItemByFileInfo(FileInfo Info)
+		{
+			FileReference Location = new FileReference(Info);
+
+			FileItem? Result;
+			if (!UniqueSourceFileMap.TryGetValue(Location, out Result))
+			{
+				FileItem NewFileItem = new FileItem(Location, Info);
+				if(UniqueSourceFileMap.TryAdd(Location, NewFileItem))
+				{
+					Result = NewFileItem;
+				}
+				else
+				{
+					Result = UniqueSourceFileMap[Location];
+				}
+			}
+			return Result;
+		}
+
+		/// <summary>
+		/// Gets a FileItem for a given path
+		/// </summary>
+		/// <param name="Location">Location of the file</param>
+		/// <returns>The FileItem that represents the given a full file path.</returns>
+		public static FileItem GetItemByFileReference(FileReference Location)
+		{
+			FileItem? Result;
+			if (!UniqueSourceFileMap.TryGetValue(Location, out Result))
+			{
+				FileItem NewFileItem = new FileItem(Location, Location.ToFileInfo());
+				if(UniqueSourceFileMap.TryAdd(Location, NewFileItem))
+				{
+					Result = NewFileItem;
+				}
+				else
+				{
+					Result = UniqueSourceFileMap[Location];
+				}
+			}
+			return Result;
+		}
+
+		/// <summary>
+		/// Deletes the file.
+		/// </summary>
+		public void Delete()
+		{
+			Debug.Assert(Exists);
+
+			int MaxRetryCount = 3;
+			int DeleteTryCount = 0;
+			bool bFileDeletedSuccessfully = false;
+			do
+			{
+				// If this isn't the first time through, sleep a little before trying again
+				if (DeleteTryCount > 0)
+				{
+					Thread.Sleep(1000);
+				}
+				DeleteTryCount++;
+				try
+				{
+					// Delete the destination file if it exists
+					FileInfo DeletedFileInfo = new FileInfo(AbsolutePath);
+					if (DeletedFileInfo.Exists)
+					{
+						DeletedFileInfo.IsReadOnly = false;
+						DeletedFileInfo.Delete();
+					}
+					// Success!
+					bFileDeletedSuccessfully = true;
+				}
+				catch (Exception Ex)
+				{
+					Log.TraceInformation("Failed to delete file '" + AbsolutePath + "'");
+					Log.TraceInformation("    Exception: " + Ex.Message);
+					if (DeleteTryCount < MaxRetryCount)
+					{
+						Log.TraceInformation("Attempting to retry...");
+					}
+					else
+					{
+						Log.TraceInformation("ERROR: Exhausted all retries!");
+					}
+				}
+			}
+			while (!bFileDeletedSuccessfully && (DeleteTryCount < MaxRetryCount));
+		}
+
+		/// <summary>
+		/// Resets the cached file info
+		/// </summary>
+		public void ResetCachedInfo()
+		{
+			Info = Location.ToFileInfo();
+		}
+
+		/// <summary>
+		/// Resets all cached file info. Significantly reduces performance; do not use unless strictly necessary.
+		/// </summary>
+		public static void ResetAllCachedInfo_SLOW()
+		{
+			foreach(FileItem Item in UniqueSourceFileMap.Values)
+			{
+				Item.ResetCachedInfo();
+			}
+		}
+
+		/// <summary>
+		/// Return the path to this FileItem to debugging
+		/// </summary>
+		/// <returns>Absolute path to this file item</returns>
+		public override string ToString()
+		{
+			return AbsolutePath;
+		}
+	}
+
+	/// <summary>
+	/// Helper functions for serialization
+	/// </summary>
+	public static class FileItemExtensionMethods
+	{
+		/// <summary>
+		/// Read a file item from a binary archive
+		/// </summary>
+		/// <param name="Reader">Reader to serialize data from</param>
+		/// <returns>Instance of the serialized file item</returns>
+		public static FileItem ReadFileItem(this BinaryArchiveReader Reader)
+		{
+			return Reader.ReadObjectReference<FileItem>(() => FileItem.GetItemByFileReference(Reader.ReadFileReference()));
+		}
+
+		/// <summary>
+		/// Write a file item to a binary archive
+		/// </summary>
+		/// <param name="Writer">Writer to serialize data to</param>
+		/// <param name="FileItem">File item to write</param>
+		public static void WriteFileItem(this BinaryArchiveWriter Writer, FileItem? FileItem)
+		{
+			Writer.WriteObjectReference<FileItem>(FileItem!, () => Writer.WriteFileReference(FileItem!.Location));
+		}
+
+		/// <summary>
+		/// Read a file item as a DirectoryItem and name. This is slower than reading it directly, but results in a significantly smaller archive
+		/// where most files are in the same directories.
+		/// </summary>
+		/// <param name="Reader">Archive to read from</param>
+		/// <returns>FileItem read from the archive</returns>
+		static FileItem ReadCompactFileItemData(this BinaryArchiveReader Reader)
+		{
+			DirectoryItem Directory = Reader.ReadDirectoryItem();
+			string Name = Reader.ReadString();
+
+			FileItem FileItem = FileItem.GetItemByFileReference(FileReference.Combine(Directory.Location, Name));
+            FileItem.UpdateCachedDirectory(Directory);
+            return FileItem;
+		}
+
+		/// <summary>
+		/// Read a file item in a format which de-duplicates directory names.
+		/// </summary>
+		/// <param name="Reader">Reader to serialize data from</param>
+		/// <returns>Instance of the serialized file item</returns>
+		public static FileItem ReadCompactFileItem(this BinaryArchiveReader Reader)
+		{
+			return Reader.ReadObjectReference<FileItem>(() => ReadCompactFileItemData(Reader));
+		}
+
+		/// <summary>
+		/// Writes a file item in a format which de-duplicates directory names.
+		/// </summary>
+		/// <param name="Writer">Writer to serialize data to</param>
+		/// <param name="FileItem">File item to write</param>
+		public static void WriteCompactFileItem(this BinaryArchiveWriter Writer, FileItem FileItem)
+		{
+			Writer.WriteObjectReference<FileItem>(FileItem, () => { Writer.WriteDirectoryItem(FileItem.GetDirectoryItem()); Writer.WriteString(FileItem.Name); });
+		}
+	}
+}
diff --git a/Engine/Source/Programs/DotNETCommon/BuildUtilities/UEBuildPlatformSDK.cs b/Engine/Source/Programs/DotNETCommon/BuildUtilities/UEBuildPlatformSDK.cs
--- a/Engine/Source/Programs/DotNETCommon/BuildUtilities/UEBuildPlatformSDK.cs
+++ b/Engine/Source/Programs/DotNETCommon/BuildUtilities/UEBuildPlatformSDK.cs
@@ -7,6 +7,8 @@ using System.IO;
 using System.Linq;
 using System.Text;
 
+#nullable disable
+
 namespace EpicGames.Core
 {
 	/// <summary>
diff --git a/Engine/Source/Programs/UnrealBuildTool/Configuration/UEBuildBinary.cs b/Engine/Source/Programs/UnrealBuildTool/Configuration/UEBuildBinary.cs
--- a/Engine/Source/Programs/UnrealBuildTool/Configuration/UEBuildBinary.cs
+++ b/Engine/Source/Programs/UnrealBuildTool/Configuration/UEBuildBinary.cs
@@ -9,6 +9,7 @@ using System.Linq;
 using System.Xml;
 using System.Globalization;
 using EpicGames.Core;
+using UnrealBuildBase;
 
 namespace UnrealBuildTool
 {
diff --git a/Engine/Source/Programs/UnrealBuildTool/Configuration/UEBuildFramework.cs b/Engine/Source/Programs/UnrealBuildTool/Configuration/UEBuildFramework.cs
--- a/Engine/Source/Programs/UnrealBuildTool/Configuration/UEBuildFramework.cs
+++ b/Engine/Source/Programs/UnrealBuildTool/Configuration/UEBuildFramework.cs
@@ -6,6 +6,7 @@ using System.Linq;
 using System.Text;
 using System.Threading.Tasks;
 using EpicGames.Core;
+using UnrealBuildBase;
 
 namespace UnrealBuildTool
 {
diff --git a/Engine/Source/Programs/UnrealBuildTool/Configuration/UEBuildModule.cs b/Engine/Source/Programs/UnrealBuildTool/Configuration/UEBuildModule.cs
--- a/Engine/Source/Programs/UnrealBuildTool/Configuration/UEBuildModule.cs
+++ b/Engine/Source/Programs/UnrealBuildTool/Configuration/UEBuildModule.cs
@@ -9,7 +9,7 @@ using System.Linq;
 using System.Text.RegularExpressions;
 using System.Xml;
 using EpicGames.Core;
-
+using UnrealBuildBase;
 
 namespace UnrealBuildTool
 {
diff --git a/Engine/Source/Programs/UnrealBuildTool/Configuration/UEBuildModuleCPP.cs b/Engine/Source/Programs/UnrealBuildTool/Configuration/UEBuildModuleCPP.cs
--- a/Engine/Source/Programs/UnrealBuildTool/Configuration/UEBuildModuleCPP.cs
+++ b/Engine/Source/Programs/UnrealBuildTool/Configuration/UEBuildModuleCPP.cs
@@ -8,6 +8,7 @@ using System.Linq;
 using System.Text;
 using System.Threading.Tasks;
 using EpicGames.Core;
+using UnrealBuildBase;
 
 namespace UnrealBuildTool
 {
diff --git a/Engine/Source/Programs/UnrealBuildTool/Configuration/UEBuildTarget.cs b/Engine/Source/Programs/UnrealBuildTool/Configuration/UEBuildTarget.cs
--- a/Engine/Source/Programs/UnrealBuildTool/Configuration/UEBuildTarget.cs
+++ b/Engine/Source/Programs/UnrealBuildTool/Configuration/UEBuildTarget.cs
@@ -13,6 +13,7 @@ using System.Runtime.Serialization;
 using System.Text.RegularExpressions;
 using EpicGames.Core;
 using System.Reflection;
+using UnrealBuildBase;
 
 namespace UnrealBuildTool
 {
diff --git a/Engine/Source/Programs/UnrealBuildTool/Executors/Experimental/FASTBuild.cs b/Engine/Source/Programs/UnrealBuildTool/Executors/Experimental/FASTBuild.cs
--- a/Engine/Source/Programs/UnrealBuildTool/Executors/Experimental/FASTBuild.cs
+++ b/Engine/Source/Programs/UnrealBuildTool/Executors/Experimental/FASTBuild.cs
@@ -21,6 +21,7 @@ using System.Linq;
 using EpicGames.Core;
 using System.Runtime.Serialization;
 using System.Text.RegularExpressions;
+using UnrealBuildBase;
 
 #nullable disable
 
diff --git a/Engine/Source/Programs/UnrealBuildTool/Executors/XGE.cs b/Engine/Source/Programs/UnrealBuildTool/Executors/XGE.cs
--- a/Engine/Source/Programs/UnrealBuildTool/Executors/XGE.cs
+++ b/Engine/Source/Programs/UnrealBuildTool/Executors/XGE.cs
@@ -15,6 +15,7 @@ using EpicGames.Core;
 using System.Net;
 using System.Net.NetworkInformation;
 using System.Runtime.InteropServices;
+using UnrealBuildBase;
 
 #nullable disable
 
diff --git a/Engine/Source/Programs/UnrealBuildTool/Modes/BuildMode.cs b/Engine/Source/Programs/UnrealBuildTool/Modes/BuildMode.cs
--- a/Engine/Source/Programs/UnrealBuildTool/Modes/BuildMode.cs
+++ b/Engine/Source/Programs/UnrealBuildTool/Modes/BuildMode.cs
@@ -9,6 +9,7 @@ using System.Text;
 using System.Threading;
 using System.Threading.Tasks;
 using EpicGames.Core;
+using UnrealBuildBase;
 
 namespace UnrealBuildTool
 {
diff --git a/Engine/Source/Programs/UnrealBuildTool/Modes/GenerateClangDatabase.cs b/Engine/Source/Programs/UnrealBuildTool/Modes/GenerateClangDatabase.cs
--- a/Engine/Source/Programs/UnrealBuildTool/Modes/GenerateClangDatabase.cs
+++ b/Engine/Source/Programs/UnrealBuildTool/Modes/GenerateClangDatabase.cs
@@ -6,6 +6,7 @@ using System.Linq;
 using System.Text;
 using System.Threading.Tasks;
 using EpicGames.Core;
+using UnrealBuildBase;
 
 #nullable disable
 
diff --git a/Engine/Source/Programs/UnrealBuildTool/Modes/JsonExportMode.cs b/Engine/Source/Programs/UnrealBuildTool/Modes/JsonExportMode.cs
--- a/Engine/Source/Programs/UnrealBuildTool/Modes/JsonExportMode.cs
+++ b/Engine/Source/Programs/UnrealBuildTool/Modes/JsonExportMode.cs
@@ -6,6 +6,7 @@ using System.Linq;
 using System.Text;
 using System.Threading.Tasks;
 using EpicGames.Core;
+using UnrealBuildBase;
 
 namespace UnrealBuildTool
 {
diff --git a/Engine/Source/Programs/UnrealBuildTool/Platform/Android/AndroidToolChain.cs b/Engine/Source/Programs/UnrealBuildTool/Platform/Android/AndroidToolChain.cs
--- a/Engine/Source/Programs/UnrealBuildTool/Platform/Android/AndroidToolChain.cs
+++ b/Engine/Source/Programs/UnrealBuildTool/Platform/Android/AndroidToolChain.cs
@@ -8,6 +8,7 @@ using System.IO;
 using System.Linq;
 using Microsoft.Win32;
 using EpicGames.Core;
+using UnrealBuildBase;
 
 #nullable disable
 
diff --git a/Engine/Source/Programs/UnrealBuildTool/Platform/HoloLens/HoloLensToolChain.cs b/Engine/Source/Programs/UnrealBuildTool/Platform/HoloLens/HoloLensToolChain.cs
--- a/Engine/Source/Programs/UnrealBuildTool/Platform/HoloLens/HoloLensToolChain.cs
+++ b/Engine/Source/Programs/UnrealBuildTool/Platform/HoloLens/HoloLensToolChain.cs
@@ -9,6 +9,7 @@ using System.Text;
 using System.Text.RegularExpressions;
 using Microsoft.Win32;
 using EpicGames.Core;
+using UnrealBuildBase;
 
 #nullable disable
 
diff --git a/Engine/Source/Programs/UnrealBuildTool/Platform/IOS/IOSToolChain.cs b/Engine/Source/Programs/UnrealBuildTool/Platform/IOS/IOSToolChain.cs
--- a/Engine/Source/Programs/UnrealBuildTool/Platform/IOS/IOSToolChain.cs
+++ b/Engine/Source/Programs/UnrealBuildTool/Platform/IOS/IOSToolChain.cs
@@ -15,6 +15,7 @@ using Ionic.Zip;
 using Ionic.Zlib;
 using EpicGames.Core;
 using System.Security.Cryptography.X509Certificates;
+using UnrealBuildBase;
 
 #nullable disable
 
diff --git a/Engine/Source/Programs/UnrealBuildTool/Platform/Linux/LinuxToolChain.cs b/Engine/Source/Programs/UnrealBuildTool/Platform/Linux/LinuxToolChain.cs
--- a/Engine/Source/Programs/UnrealBuildTool/Platform/Linux/LinuxToolChain.cs
+++ b/Engine/Source/Programs/UnrealBuildTool/Platform/Linux/LinuxToolChain.cs
@@ -8,6 +8,7 @@ using System.IO;
 using System.Linq;
 using Microsoft.Win32;
 using EpicGames.Core;
+using UnrealBuildBase;
 
 #nullable disable
 
diff --git a/Engine/Source/Programs/UnrealBuildTool/Platform/Lumin/LuminToolChain.cs b/Engine/Source/Programs/UnrealBuildTool/Platform/Lumin/LuminToolChain.cs
--- a/Engine/Source/Programs/UnrealBuildTool/Platform/Lumin/LuminToolChain.cs
+++ b/Engine/Source/Programs/UnrealBuildTool/Platform/Lumin/LuminToolChain.cs
@@ -8,6 +8,7 @@ using System.IO;
 using System.Linq;
 using Microsoft.Win32;
 using EpicGames.Core;
+using UnrealBuildBase;
 
 #nullable disable
 
diff --git a/Engine/Source/Programs/UnrealBuildTool/Platform/Mac/MacToolChain.cs b/Engine/Source/Programs/UnrealBuildTool/Platform/Mac/MacToolChain.cs
--- a/Engine/Source/Programs/UnrealBuildTool/Platform/Mac/MacToolChain.cs
+++ b/Engine/Source/Programs/UnrealBuildTool/Platform/Mac/MacToolChain.cs
@@ -12,6 +12,7 @@ using Ionic.Zip;
 using EpicGames.Core;
 using System.Globalization;
 using System.Text.RegularExpressions;
+using UnrealBuildBase;
 
 #nullable disable
 
diff --git a/Engine/Source/Programs/UnrealBuildTool/Platform/Windows/PVSToolChain.cs b/Engine/Source/Programs/UnrealBuildTool/Platform/Windows/PVSToolChain.cs
--- a/Engine/Source/Programs/UnrealBuildTool/Platform/Windows/PVSToolChain.cs
+++ b/Engine/Source/Programs/UnrealBuildTool/Platform/Windows/PVSToolChain.cs
@@ -12,6 +12,7 @@ using System.Threading.Tasks;
 using System.Xml;
 using System.Xml.Serialization;
 using EpicGames.Core;
+using UnrealBuildBase;
 
 #nullable disable
 
diff --git a/Engine/Source/Programs/UnrealBuildTool/Platform/Windows/VCCompileAction.cs b/Engine/Source/Programs/UnrealBuildTool/Platform/Windows/VCCompileAction.cs
--- a/Engine/Source/Programs/UnrealBuildTool/Platform/Windows/VCCompileAction.cs
+++ b/Engine/Source/Programs/UnrealBuildTool/Platform/Windows/VCCompileAction.cs
@@ -3,6 +3,7 @@
 using System;
 using System.Collections.Generic;
 using EpicGames.Core;
+using UnrealBuildBase;
 
 #nullable disable
 
diff --git a/Engine/Source/Programs/UnrealBuildTool/Platform/Windows/VCToolChain.cs b/Engine/Source/Programs/UnrealBuildTool/Platform/Windows/VCToolChain.cs
--- a/Engine/Source/Programs/UnrealBuildTool/Platform/Windows/VCToolChain.cs
+++ b/Engine/Source/Programs/UnrealBuildTool/Platform/Windows/VCToolChain.cs
@@ -9,6 +9,7 @@ using System.Linq;
 using Microsoft.Win32;
 using System.Text;
 using EpicGames.Core;
+using UnrealBuildBase;
 
 #nullable disable
 
diff --git a/Engine/Source/Programs/UnrealBuildTool/Preprocessor/Preprocessor.cs b/Engine/Source/Programs/UnrealBuildTool/Preprocessor/Preprocessor.cs
--- a/Engine/Source/Programs/UnrealBuildTool/Preprocessor/Preprocessor.cs
+++ b/Engine/Source/Programs/UnrealBuildTool/Preprocessor/Preprocessor.cs
@@ -11,6 +11,7 @@ using System.Collections;
 using System.Collections.Concurrent;
 using EpicGames.Core;
 using System.Diagnostics.CodeAnalysis;
+using UnrealBuildBase;
 
 namespace UnrealBuildTool
 {
diff --git a/Engine/Source/Programs/UnrealBuildTool/Preprocessor/PreprocessorContext.cs b/Engine/Source/Programs/UnrealBuildTool/Preprocessor/PreprocessorContext.cs
--- a/Engine/Source/Programs/UnrealBuildTool/Preprocessor/PreprocessorContext.cs
+++ b/Engine/Source/Programs/UnrealBuildTool/Preprocessor/PreprocessorContext.cs
@@ -5,6 +5,7 @@ using System.Collections.Generic;
 using System.Linq;
 using System.Text;
 using System.Threading.Tasks;
+using UnrealBuildBase;
 
 namespace UnrealBuildTool
 {
diff --git a/Engine/Source/Programs/UnrealBuildTool/Preprocessor/SourceFile.cs b/Engine/Source/Programs/UnrealBuildTool/Preprocessor/SourceFile.cs
--- a/Engine/Source/Programs/UnrealBuildTool/Preprocessor/SourceFile.cs
+++ b/Engine/Source/Programs/UnrealBuildTool/Preprocessor/SourceFile.cs
@@ -9,6 +9,7 @@ using System.Linq;
 using System.Text;
 using System.Threading.Tasks;
 using EpicGames.Core;
+using UnrealBuildBase;
 
 namespace UnrealBuildTool
 {
diff --git a/Engine/Source/Programs/UnrealBuildTool/System/Action.cs b/Engine/Source/Programs/UnrealBuildTool/System/Action.cs
--- a/Engine/Source/Programs/UnrealBuildTool/System/Action.cs
+++ b/Engine/Source/Programs/UnrealBuildTool/System/Action.cs
@@ -8,6 +8,7 @@ using System.Reflection;
 using System.Text;
 using System.Threading.Tasks;
 using EpicGames.Core;
+using UnrealBuildBase;
 
 namespace UnrealBuildTool
 {
diff --git a/Engine/Source/Programs/UnrealBuildTool/System/ActionGraph.cs b/Engine/Source/Programs/UnrealBuildTool/System/ActionGraph.cs
--- a/Engine/Source/Programs/UnrealBuildTool/System/ActionGraph.cs
+++ b/Engine/Source/Programs/UnrealBuildTool/System/ActionGraph.cs
@@ -11,6 +11,7 @@ using System.Text;
 using System.Threading.Tasks;
 using System.Text.Json;
 using System.Text.Json.Serialization;
+using UnrealBuildBase;
 
 namespace UnrealBuildTool
 {
diff --git a/Engine/Source/Programs/UnrealBuildTool/System/ActionGraphBuilder.cs b/Engine/Source/Programs/UnrealBuildTool/System/ActionGraphBuilder.cs
--- a/Engine/Source/Programs/UnrealBuildTool/System/ActionGraphBuilder.cs
+++ b/Engine/Source/Programs/UnrealBuildTool/System/ActionGraphBuilder.cs
@@ -7,6 +7,7 @@ using System.Reflection;
 using System.Text;
 using System.Threading.Tasks;
 using EpicGames.Core;
+using UnrealBuildBase;
 
 #nullable disable
 
diff --git a/Engine/Source/Programs/UnrealBuildTool/System/ActionHistory.cs b/Engine/Source/Programs/UnrealBuildTool/System/ActionHistory.cs
--- a/Engine/Source/Programs/UnrealBuildTool/System/ActionHistory.cs
+++ b/Engine/Source/Programs/UnrealBuildTool/System/ActionHistory.cs
@@ -10,6 +10,7 @@ using System.Runtime.Serialization.Formatters.Binary;
 using System.Security.Cryptography;
 using System.Text;
 using EpicGames.Core;
+using UnrealBuildBase;
 
 namespace UnrealBuildTool
 {
diff --git a/Engine/Source/Programs/UnrealBuildTool/System/CppCompileEnvironment.cs b/Engine/Source/Programs/UnrealBuildTool/System/CppCompileEnvironment.cs
--- a/Engine/Source/Programs/UnrealBuildTool/System/CppCompileEnvironment.cs
+++ b/Engine/Source/Programs/UnrealBuildTool/System/CppCompileEnvironment.cs
@@ -7,6 +7,7 @@ using System.Text;
 using System.Text.RegularExpressions;
 using System.IO;
 using EpicGames.Core;
+using UnrealBuildBase;
 
 namespace UnrealBuildTool
 {
diff --git a/Engine/Source/Programs/UnrealBuildTool/System/CppDependencyCache.cs b/Engine/Source/Programs/UnrealBuildTool/System/CppDependencyCache.cs
--- a/Engine/Source/Programs/UnrealBuildTool/System/CppDependencyCache.cs
+++ b/Engine/Source/Programs/UnrealBuildTool/System/CppDependencyCache.cs
@@ -10,6 +10,7 @@ using System.Linq;
 using System.Text;
 using System.Threading.Tasks;
 using EpicGames.Core;
+using UnrealBuildBase;
 
 namespace UnrealBuildTool
 {
diff --git a/Engine/Source/Programs/UnrealBuildTool/System/DirectoryItem.cs b/Engine/Source/Programs/UnrealBuildTool/System/DirectoryItem.cs
deleted file mode 100644
diff --git a/Engine/Source/Programs/UnrealBuildTool/System/DirectoryLookupCache.cs b/Engine/Source/Programs/UnrealBuildTool/System/DirectoryLookupCache.cs
--- a/Engine/Source/Programs/UnrealBuildTool/System/DirectoryLookupCache.cs
+++ b/Engine/Source/Programs/UnrealBuildTool/System/DirectoryLookupCache.cs
@@ -5,6 +5,7 @@ using System.IO;
 using System.Collections.Generic;
 using EpicGames.Core;
 using System.Linq;
+using UnrealBuildBase;
 
 namespace UnrealBuildTool
 {
diff --git a/Engine/Source/Programs/UnrealBuildTool/System/DynamicCompilation.cs b/Engine/Source/Programs/UnrealBuildTool/System/DynamicCompilation.cs
--- a/Engine/Source/Programs/UnrealBuildTool/System/DynamicCompilation.cs
+++ b/Engine/Source/Programs/UnrealBuildTool/System/DynamicCompilation.cs
@@ -17,6 +17,7 @@ using Microsoft.CodeAnalysis.Emit;
 using System.Reflection.Metadata;
 using Microsoft.CodeAnalysis.Text;
 using System.Runtime.InteropServices;
+using UnrealBuildBase;
 
 namespace UnrealBuildTool
 {
diff --git a/Engine/Source/Programs/UnrealBuildTool/System/ExternalExecution.cs b/Engine/Source/Programs/UnrealBuildTool/System/ExternalExecution.cs
--- a/Engine/Source/Programs/UnrealBuildTool/System/ExternalExecution.cs
+++ b/Engine/Source/Programs/UnrealBuildTool/System/ExternalExecution.cs
@@ -13,6 +13,7 @@ using System.Runtime.InteropServices;
 using System.Runtime.Serialization;
 using EpicGames.Core;
 using System.Text.Json.Serialization;
+using UnrealBuildBase;
 
 namespace UnrealBuildTool
 {
diff --git a/Engine/Source/Programs/UnrealBuildTool/System/FileItem.cs b/Engine/Source/Programs/UnrealBuildTool/System/FileItem.cs
deleted file mode 100644
diff --git a/Engine/Source/Programs/UnrealBuildTool/System/FileMetadataPrefetch.cs b/Engine/Source/Programs/UnrealBuildTool/System/FileMetadataPrefetch.cs
--- a/Engine/Source/Programs/UnrealBuildTool/System/FileMetadataPrefetch.cs
+++ b/Engine/Source/Programs/UnrealBuildTool/System/FileMetadataPrefetch.cs
@@ -7,6 +7,7 @@ using System.Text;
 using System.Threading;
 using System.Threading.Tasks;
 using EpicGames.Core;
+using UnrealBuildBase;
 
 namespace UnrealBuildTool
 {
diff --git a/Engine/Source/Programs/UnrealBuildTool/System/HotReload.cs b/Engine/Source/Programs/UnrealBuildTool/System/HotReload.cs
--- a/Engine/Source/Programs/UnrealBuildTool/System/HotReload.cs
+++ b/Engine/Source/Programs/UnrealBuildTool/System/HotReload.cs
@@ -9,6 +9,7 @@ using System.Text;
 using System.Threading;
 using System.Threading.Tasks;
 using EpicGames.Core;
+using UnrealBuildBase;
 
 #nullable disable
 
diff --git a/Engine/Source/Programs/UnrealBuildTool/System/LinkEnvironment.cs b/Engine/Source/Programs/UnrealBuildTool/System/LinkEnvironment.cs
--- a/Engine/Source/Programs/UnrealBuildTool/System/LinkEnvironment.cs
+++ b/Engine/Source/Programs/UnrealBuildTool/System/LinkEnvironment.cs
@@ -6,6 +6,7 @@ using System.Collections.Generic;
 using System.Linq;
 using System.Text;
 using EpicGames.Core;
+using UnrealBuildBase;
 
 namespace UnrealBuildTool
 {
diff --git a/Engine/Source/Programs/UnrealBuildTool/System/Plugins.cs b/Engine/Source/Programs/UnrealBuildTool/System/Plugins.cs
--- a/Engine/Source/Programs/UnrealBuildTool/System/Plugins.cs
+++ b/Engine/Source/Programs/UnrealBuildTool/System/Plugins.cs
@@ -6,6 +6,7 @@ using System.IO;
 using System.Diagnostics;
 using System.Linq;
 using EpicGames.Core;
+using UnrealBuildBase;
 
 namespace UnrealBuildTool
 {
diff --git a/Engine/Source/Programs/UnrealBuildTool/System/PrecompiledHeaderInstance.cs b/Engine/Source/Programs/UnrealBuildTool/System/PrecompiledHeaderInstance.cs
--- a/Engine/Source/Programs/UnrealBuildTool/System/PrecompiledHeaderInstance.cs
+++ b/Engine/Source/Programs/UnrealBuildTool/System/PrecompiledHeaderInstance.cs
@@ -5,6 +5,7 @@ using System.Collections.Generic;
 using System.Linq;
 using System.Text;
 using System.Threading.Tasks;
+using UnrealBuildBase;
 
 namespace UnrealBuildTool
 {
diff --git a/Engine/Source/Programs/UnrealBuildTool/System/PrecompiledHeaderTemplate.cs b/Engine/Source/Programs/UnrealBuildTool/System/PrecompiledHeaderTemplate.cs
--- a/Engine/Source/Programs/UnrealBuildTool/System/PrecompiledHeaderTemplate.cs
+++ b/Engine/Source/Programs/UnrealBuildTool/System/PrecompiledHeaderTemplate.cs
@@ -6,6 +6,7 @@ using System.Linq;
 using System.Text;
 using System.Threading.Tasks;
 using EpicGames.Core;
+using UnrealBuildBase;
 
 namespace UnrealBuildTool
 {
diff --git a/Engine/Source/Programs/UnrealBuildTool/System/RulesAssembly.cs b/Engine/Source/Programs/UnrealBuildTool/System/RulesAssembly.cs
--- a/Engine/Source/Programs/UnrealBuildTool/System/RulesAssembly.cs
+++ b/Engine/Source/Programs/UnrealBuildTool/System/RulesAssembly.cs
@@ -10,6 +10,7 @@ using System.Runtime.Serialization;
 using System.Text;
 using System.Threading.Tasks;
 using EpicGames.Core;
+using UnrealBuildBase;
 
 namespace UnrealBuildTool
 {
diff --git a/Engine/Source/Programs/UnrealBuildTool/System/RulesCompiler.cs b/Engine/Source/Programs/UnrealBuildTool/System/RulesCompiler.cs
--- a/Engine/Source/Programs/UnrealBuildTool/System/RulesCompiler.cs
+++ b/Engine/Source/Programs/UnrealBuildTool/System/RulesCompiler.cs
@@ -6,6 +6,7 @@ using System.IO;
 using System.Linq;
 using System.Reflection;
 using EpicGames.Core;
+using UnrealBuildBase;
 
 namespace UnrealBuildTool
 {
diff --git a/Engine/Source/Programs/UnrealBuildTool/System/SourceFileMetadataCache.cs b/Engine/Source/Programs/UnrealBuildTool/System/SourceFileMetadataCache.cs
--- a/Engine/Source/Programs/UnrealBuildTool/System/SourceFileMetadataCache.cs
+++ b/Engine/Source/Programs/UnrealBuildTool/System/SourceFileMetadataCache.cs
@@ -10,6 +10,7 @@ using System.Text;
 using System.Text.RegularExpressions;
 using System.Threading.Tasks;
 using EpicGames.Core;
+using UnrealBuildBase;
 
 #nullable enable
 
diff --git a/Engine/Source/Programs/UnrealBuildTool/System/SourceFileWorkingSet.cs b/Engine/Source/Programs/UnrealBuildTool/System/SourceFileWorkingSet.cs
--- a/Engine/Source/Programs/UnrealBuildTool/System/SourceFileWorkingSet.cs
+++ b/Engine/Source/Programs/UnrealBuildTool/System/SourceFileWorkingSet.cs
@@ -9,6 +9,7 @@ using System.Linq;
 using System.Text;
 using System.Threading.Tasks;
 using EpicGames.Core;
+using UnrealBuildBase;
 
 namespace UnrealBuildTool
 {
diff --git a/Engine/Source/Programs/UnrealBuildTool/System/TargetMakefile.cs b/Engine/Source/Programs/UnrealBuildTool/System/TargetMakefile.cs
--- a/Engine/Source/Programs/UnrealBuildTool/System/TargetMakefile.cs
+++ b/Engine/Source/Programs/UnrealBuildTool/System/TargetMakefile.cs
@@ -10,6 +10,7 @@ using System.Runtime.Serialization.Formatters.Binary;
 using System.Text;
 using System.Threading.Tasks;
 using EpicGames.Core;
+using UnrealBuildBase;
 
 namespace UnrealBuildTool
 {
diff --git a/Engine/Source/Programs/UnrealBuildTool/System/Unity.cs b/Engine/Source/Programs/UnrealBuildTool/System/Unity.cs
--- a/Engine/Source/Programs/UnrealBuildTool/System/Unity.cs
+++ b/Engine/Source/Programs/UnrealBuildTool/System/Unity.cs
@@ -6,6 +6,7 @@ using System.Text;
 using System.IO;
 using System.Linq;
 using EpicGames.Core;
+using UnrealBuildBase;
 
 namespace UnrealBuildTool
 {
diff --git a/Engine/Source/Programs/UnrealBuildTool/System/Utils.cs b/Engine/Source/Programs/UnrealBuildTool/System/Utils.cs
--- a/Engine/Source/Programs/UnrealBuildTool/System/Utils.cs
+++ b/Engine/Source/Programs/UnrealBuildTool/System/Utils.cs
@@ -12,6 +12,7 @@ using System.Runtime.CompilerServices;
 using System.Runtime.InteropServices;
 using System.Linq;
 using EpicGames.Core;
+using UnrealBuildBase;
 
 namespace UnrealBuildTool
 {
diff --git a/Engine/Source/Programs/UnrealBuildTool/ToolChain/UEToolChain.cs b/Engine/Source/Programs/UnrealBuildTool/ToolChain/UEToolChain.cs
--- a/Engine/Source/Programs/UnrealBuildTool/ToolChain/UEToolChain.cs
+++ b/Engine/Source/Programs/UnrealBuildTool/ToolChain/UEToolChain.cs
@@ -6,6 +6,7 @@ using System.IO;
 using System.Linq;
 using System.Text.RegularExpressions;
 using EpicGames.Core;
+using UnrealBuildBase;
 
 #nullable disable
 
