Move DirectoryLookupCache.cs to BuildUtilities

diff --git a/Engine/Source/Programs/DotNETCommon/BuildUtilities/System/DirectoryLookupCache.cs b/Engine/Source/Programs/DotNETCommon/BuildUtilities/System/DirectoryLookupCache.cs
new file mode 100644
--- /dev/null
+++ b/Engine/Source/Programs/DotNETCommon/BuildUtilities/System/DirectoryLookupCache.cs
@@ -0,0 +1,38 @@
+// Copyright Epic Games, Inc. All Rights Reserved.
+
+using System;
+using System.IO;
+using System.Collections.Generic;
+using EpicGames.Core;
+using System.Linq;
+
+namespace UnrealBuildBase
+{
+	static public class DirectoryLookupCache
+	{
+		static public bool FileExists(FileReference File)
+		{
+			return FileItem.GetItemByFileReference(File).Exists;
+		}
+
+		static public bool DirectoryExists(DirectoryReference Directory)
+		{
+			return DirectoryItem.GetItemByDirectoryReference(Directory).Exists;
+		}
+
+		static public IEnumerable<FileReference> EnumerateFiles(DirectoryReference Directory)
+		{
+			return DirectoryItem.GetItemByDirectoryReference(Directory).EnumerateFiles().Select(x => x.Location);
+		}
+
+		static public IEnumerable<DirectoryReference> EnumerateDirectories(DirectoryReference Directory)
+		{
+			return DirectoryItem.GetItemByDirectoryReference(Directory).EnumerateDirectories().Select(x => x.Location);
+		}
+
+        static public void InvalidateCachedDirectory(DirectoryReference Directory)
+        {
+			DirectoryItem.GetItemByDirectoryReference(Directory).ResetCachedInfo();
+        }
+	}
+}
diff --git a/Engine/Source/Programs/UnrealBuildTool/ProjectFiles/ProjectFileGenerator.cs b/Engine/Source/Programs/UnrealBuildTool/ProjectFiles/ProjectFileGenerator.cs
--- a/Engine/Source/Programs/UnrealBuildTool/ProjectFiles/ProjectFileGenerator.cs
+++ b/Engine/Source/Programs/UnrealBuildTool/ProjectFiles/ProjectFileGenerator.cs
@@ -11,6 +11,7 @@ using System.Diagnostics;
 using EpicGames.Core;
 using System.Xml.Linq;
 using System.Xml;
+using UnrealBuildBase;
 
 #nullable disable
 
diff --git a/Engine/Source/Programs/UnrealBuildTool/System/DirectoryLookupCache.cs b/Engine/Source/Programs/UnrealBuildTool/System/DirectoryLookupCache.cs
deleted file mode 100644
diff --git a/Engine/Source/Programs/UnrealBuildTool/System/NativeProjects.cs b/Engine/Source/Programs/UnrealBuildTool/System/NativeProjects.cs
--- a/Engine/Source/Programs/UnrealBuildTool/System/NativeProjects.cs
+++ b/Engine/Source/Programs/UnrealBuildTool/System/NativeProjects.cs
@@ -8,6 +8,7 @@ using System.Linq;
 using System.Text;
 using System.Threading.Tasks;
 using EpicGames.Core;
+using UnrealBuildBase;
 
 namespace UnrealBuildTool
 {
