UnrealBuildTool:

Move StartupTraceListener.cs to BuildUtilities

diff --git a/Engine/Source/Programs/AutomationTool/AutomationUtils/Automation.cs b/Engine/Source/Programs/AutomationTool/AutomationUtils/Automation.cs
--- a/Engine/Source/Programs/AutomationTool/AutomationUtils/Automation.cs
+++ b/Engine/Source/Programs/AutomationTool/AutomationUtils/Automation.cs
@@ -8,6 +8,7 @@ using System.Reflection;
 using System.Diagnostics;
 using UnrealBuildTool;
 using EpicGames.Core;
+using UnrealBuildBase;
 
 namespace AutomationTool
 {
diff --git a/Engine/Source/Programs/AutomationTool/Program.cs b/Engine/Source/Programs/AutomationTool/Program.cs
--- a/Engine/Source/Programs/AutomationTool/Program.cs
+++ b/Engine/Source/Programs/AutomationTool/Program.cs
@@ -6,12 +6,12 @@
 using System;
 using System.Threading;
 using System.Diagnostics;
-using UnrealBuildTool;
 using System.Reflection;
 using EpicGames.Core;
 using System.IO;
 using System.Collections.Generic;
 using System.Text;
+using UnrealBuildBase;
 
 namespace AutomationTool
 {
diff --git a/Engine/Source/Programs/DotNETCommon/BuildUtilities/System/StartupTraceListener.cs b/Engine/Source/Programs/DotNETCommon/BuildUtilities/System/StartupTraceListener.cs
new file mode 100644
--- /dev/null
+++ b/Engine/Source/Programs/DotNETCommon/BuildUtilities/System/StartupTraceListener.cs
@@ -0,0 +1,73 @@
+// Copyright Epic Games, Inc. All Rights Reserved.
+
+using System;
+using System.Collections.Generic;
+using System.Linq;
+using System.Text;
+using System.IO;
+using System.CodeDom.Compiler;
+using Microsoft.CSharp;
+using Microsoft.Win32;
+using System.Reflection;
+using System.Diagnostics;
+using EpicGames.Core;
+
+namespace UnrealBuildBase
+{
+	/// <summary>
+	/// Captures all log output during startup until a log file writer has been created
+	/// </summary>
+	public class StartupTraceListener : TraceListener
+	{
+		StringBuilder Buffer = new StringBuilder();
+
+		/// <summary>
+		/// Constructor
+		/// </summary>
+		public StartupTraceListener()
+		{
+		}
+
+		/// <summary>
+		/// Copy the contents of the buffered output to another trace listener
+		/// </summary>
+		/// <param name="Other">The trace listener to receive the buffered output</param>
+		public void CopyTo(TraceListener Other)
+		{
+			int StartIdx = 0;
+			while(StartIdx < Buffer.Length)
+			{
+				int EndIdx = StartIdx;
+				while(EndIdx < Buffer.Length && Buffer[EndIdx] != '\n')
+				{
+					EndIdx++;
+				}
+				Other.WriteLine(Buffer.ToString(StartIdx, EndIdx - StartIdx));
+				StartIdx = EndIdx + 1;
+			}
+		}
+
+		/// <summary>
+		/// Write a message to the buffer
+		/// </summary>
+		/// <param name="Message">The message to write</param>
+		public override void Write(string Message)
+		{
+			if(NeedIndent)
+			{
+				WriteIndent();
+			}
+			Buffer.Append(Message);
+		}
+
+		/// <summary>
+		/// Write a message to the buffer, followed by a newline
+		/// </summary>
+		/// <param name="Message">The message to write</param>
+		public override void WriteLine(string Message)
+		{
+			Write(Message);
+			Buffer.Append("\n");
+		}
+	}
+}
diff --git a/Engine/Source/Programs/UnrealBuildTool/Modes/GenerateProjectFilesMode.cs b/Engine/Source/Programs/UnrealBuildTool/Modes/GenerateProjectFilesMode.cs
--- a/Engine/Source/Programs/UnrealBuildTool/Modes/GenerateProjectFilesMode.cs
+++ b/Engine/Source/Programs/UnrealBuildTool/Modes/GenerateProjectFilesMode.cs
@@ -9,6 +9,7 @@ using System.Text;
 using System.Threading.Tasks;
 using EpicGames.Core;
 using System.Diagnostics;
+using UnrealBuildBase;
 
 #nullable disable
 
diff --git a/Engine/Source/Programs/UnrealBuildTool/System/StartupTraceListener.cs b/Engine/Source/Programs/UnrealBuildTool/System/StartupTraceListener.cs
deleted file mode 100644
