BuildUtilities:

Ensure UnrealBuild.RootDirectory casing is consistent.

diff --git a/Engine/Source/Programs/DotNETCommon/BuildUtilities/UnrealBuild.cs b/Engine/Source/Programs/DotNETCommon/BuildUtilities/UnrealBuild.cs
--- a/Engine/Source/Programs/DotNETCommon/BuildUtilities/UnrealBuild.cs
+++ b/Engine/Source/Programs/DotNETCommon/BuildUtilities/UnrealBuild.cs
@@ -47,6 +47,8 @@ namespace UnrealBuildBase
 				throw new Exception($"The BuildUtilities assembly requires that applications are launched from a path containing \"Engine/Binaries/DotNET\". This application was launched from {Path.GetDirectoryName(AssemblyLocation)}");
 			}
 
+			FoundRootDirectory = DirectoryReference.FindCorrectCase(FoundRootDirectory);
+
 			// Confirm that we've found a valid root directory, by testing for the existence of a well-known file
 			FileReference ExpectedExistingFile = FileReference.Combine(FoundRootDirectory, "Engine", "Build", "Build.version");
 			if (!FileReference.Exists(ExpectedExistingFile))
@@ -65,14 +67,14 @@ namespace UnrealBuildBase
 			// the UnrealBuildTool executable is assumed to be located under {RootDirectory}/Engine/Binaries/DotNET/UnrealBuildTool/
 			FileReference UnrealBuildToolPath = FileReference.Combine(EngineDirectory, "Binaries", "DotNET", "UnrealBuildTool", UBTName);
 			
-			FileReference UnrealBuildToolPathCorrected = FileReference.FindCorrectCase(UnrealBuildToolPath);
+			UnrealBuildToolPath = FileReference.FindCorrectCase(UnrealBuildToolPath);
 
-			if (!FileReference.Exists(UnrealBuildToolPathCorrected))
+			if (!FileReference.Exists(UnrealBuildToolPath))
 			{
-				throw new Exception($"Unable to find {UBTName} in the expected location at {UnrealBuildToolPathCorrected.FullName}");
+				throw new Exception($"Unable to find {UBTName} in the expected location at {UnrealBuildToolPath.FullName}");
 			}
 
-			return UnrealBuildToolPathCorrected;
+			return UnrealBuildToolPath;
 		}
 
 		/// <summary>
