ExitCode.cs:

Copyright notice

#trivial

diff --git a/Engine/Source/Programs/DotNETCommon/BuildUtilities/Automation/ExitCode.cs b/Engine/Source/Programs/DotNETCommon/BuildUtilities/Automation/ExitCode.cs
--- a/Engine/Source/Programs/DotNETCommon/BuildUtilities/Automation/ExitCode.cs
+++ b/Engine/Source/Programs/DotNETCommon/BuildUtilities/Automation/ExitCode.cs
@@ -1,3 +1,5 @@
+// Copyright Epic Games, Inc. All Rights Reserved.
+
 using System;
 using System.Collections.Generic;
 using System.Text;
