BuildUtilities:

Add Unreal.DotnetDirectory

#trivial

diff --git a/Engine/Source/Programs/DotNETCommon/BuildUtilities/Unreal.cs b/Engine/Source/Programs/DotNETCommon/BuildUtilities/Unreal.cs
--- a/Engine/Source/Programs/DotNETCommon/BuildUtilities/Unreal.cs
+++ b/Engine/Source/Programs/DotNETCommon/BuildUtilities/Unreal.cs
@@ -82,6 +82,20 @@ namespace UnrealBuildBase
 			return UnrealBuildToolPath;
 		}
 
+		static private DirectoryReference FindDotnetDirectory()
+		{
+			string HostDotNetDirectoryName;
+			switch (RuntimePlatform.Current)
+			{
+				case RuntimePlatform.Type.Windows: HostDotNetDirectoryName = "Windows"; break;
+				case RuntimePlatform.Type.Mac:     HostDotNetDirectoryName = "Mac"; break;
+				case RuntimePlatform.Type.Linux:   HostDotNetDirectoryName = "Linux"; break;
+				default: throw new Exception("Unknown host platform");
+			}
+
+			return DirectoryReference.Combine(EngineDirectory, "Binaries", "ThirdParty", "DotNet", HostDotNetDirectoryName);
+		}
+
 		/// <summary>
 		/// The full name of the root UE directory
 		/// </summary>
@@ -97,6 +111,11 @@ namespace UnrealBuildBase
 		/// </summary>
 		public static readonly FileReference UnrealBuildToolPath = FindUnrealBuildTool();
 
+		/// <summary>
+		/// The directory containing the bundled .NET installation
+		/// </summary>
+		static public readonly DirectoryReference DotnetDirectory = FindDotnetDirectory();
+
 		/// <summary>
 		/// Whether we're running with engine installed
 		/// </summary>
