BuildUtilities:

- Add Unreal.DotnetPath
- in Unreal.FindRootDirectory(), fix assembly path case sooner (allows UBT to run even if I type 'Engine' as 'engine')

EpicGames.Core:
- Add RuntimePlatform.ExeExtension

#trivial

diff --git a/Engine/Source/Programs/DotNETCommon/BuildUtilities/Unreal.cs b/Engine/Source/Programs/DotNETCommon/BuildUtilities/Unreal.cs
--- a/Engine/Source/Programs/DotNETCommon/BuildUtilities/Unreal.cs
+++ b/Engine/Source/Programs/DotNETCommon/BuildUtilities/Unreal.cs
@@ -24,7 +24,7 @@ namespace UnrealBuildBase
 			// Use the EntryAssembly (the application path), rather than the ExecutingAssembly (the library path)
 			string AssemblyLocation = Assembly.GetEntryAssembly().GetOriginalLocation();
 
-			DirectoryReference? FoundRootDirectory = DirectoryReference.FromString(AssemblyLocation);
+			DirectoryReference? FoundRootDirectory = DirectoryReference.FindCorrectCase(DirectoryReference.FromString(AssemblyLocation)!);
 
 			// Search up through the directory tree for the deepest instance of the sub-path "Engine/Binaries/DotNET"
 			while(FoundRootDirectory != null)
@@ -52,8 +52,6 @@ namespace UnrealBuildBase
 				throw new Exception($"The BuildUtilities assembly requires that applications are launched from a path containing \"Engine/Binaries/DotNET\". This application was launched from {Path.GetDirectoryName(AssemblyLocation)}");
 			}
 
-			FoundRootDirectory = DirectoryReference.FindCorrectCase(FoundRootDirectory);
-
 			// Confirm that we've found a valid root directory, by testing for the existence of a well-known file
 			FileReference ExpectedExistingFile = FileReference.Combine(FoundRootDirectory, "Engine", "Build", "Build.version");
 			if (!FileReference.Exists(ExpectedExistingFile))
@@ -67,7 +65,7 @@ namespace UnrealBuildBase
 		private static FileReference FindUnrealBuildTool()
 		{
 			// todo: use UnrealBuildTool.dll (same on all platforms). Will require changes wherever UnrealBuildTool is invoked.
-			string UBTName = RuntimePlatform.IsWindows ? "UnrealBuildTool.exe" : "UnrealBuildTool";
+			string UBTName = "UnrealBuildTool" + RuntimePlatform.ExeExtension;
 
 			// the UnrealBuildTool executable is assumed to be located under {RootDirectory}/Engine/Binaries/DotNET/UnrealBuildTool/
 			FileReference UnrealBuildToolPath = FileReference.Combine(EngineDirectory, "Binaries", "DotNET", "UnrealBuildTool", UBTName);
@@ -116,6 +114,11 @@ namespace UnrealBuildBase
 		/// </summary>
 		static public readonly DirectoryReference DotnetDirectory = FindDotnetDirectory();
 
+		/// <summary>
+		/// The path of the bundled dotnet executable
+		/// </summary>
+		static public readonly FileReference DotnetPath = FileReference.Combine(DotnetDirectory, "dotnet" + RuntimePlatform.ExeExtension);
+
 		/// <summary>
 		/// Whether we're running with engine installed
 		/// </summary>
diff --git a/Engine/Source/Programs/Shared/EpicGames.Core/RuntimePlatform.cs b/Engine/Source/Programs/Shared/EpicGames.Core/RuntimePlatform.cs
--- a/Engine/Source/Programs/Shared/EpicGames.Core/RuntimePlatform.cs
+++ b/Engine/Source/Programs/Shared/EpicGames.Core/RuntimePlatform.cs
@@ -33,5 +33,10 @@ namespace EpicGames.Core
 		/// The current runtime platform
 		/// </summary>
 		public static readonly Type Current = IsWindows ? Type.Windows : IsMac ? Type.Mac : Type.Linux;
+
+		/// <summary>
+		/// The extension executables have on the current platform
+		/// </summary>
+		public static readonly string ExeExtension = IsWindows ? ".exe" : "";
 	}
 }
