StartupTraceListener:

Improve performance of StartupTraceListener.CopyTo() for very large trace logs by generating a string once, then iterating through the lines. Per-character iteration of a StringBuilder can perform very poorly.

#trivial

diff --git a/Engine/Source/Programs/DotNETCommon/BuildUtilities/System/StartupTraceListener.cs b/Engine/Source/Programs/DotNETCommon/BuildUtilities/System/StartupTraceListener.cs
--- a/Engine/Source/Programs/DotNETCommon/BuildUtilities/System/StartupTraceListener.cs
+++ b/Engine/Source/Programs/DotNETCommon/BuildUtilities/System/StartupTraceListener.cs
@@ -34,16 +34,9 @@ namespace UnrealBuildBase
 		/// <param name="Other">The trace listener to receive the buffered output</param>
 		public void CopyTo(TraceListener Other)
 		{
-			int StartIdx = 0;
-			while(StartIdx < Buffer.Length)
+			foreach(string Line in Buffer.ToString().Split("\n"))
 			{
-				int EndIdx = StartIdx;
-				while(EndIdx < Buffer.Length && Buffer[EndIdx] != '\n')
-				{
-					EndIdx++;
-				}
-				Other.WriteLine(Buffer.ToString(StartIdx, EndIdx - StartIdx));
-				StartIdx = EndIdx + 1;
+				Other.WriteLine(Line);
 			}
 		}
 
