Cache the xcode version obtained from xcodebuild

Each invocation of `xcodebuild -version` can take > 100ms. Caching the result reduces startup time of UnrealBuildTool and AutomationTool by 1-2 seconds.

diff --git a/Engine/Source/Programs/AutomationTool/Scripts/CheckXcodeVersion.cs b/Engine/Source/Programs/AutomationTool/Scripts/CheckXcodeVersion.cs
--- a/Engine/Source/Programs/AutomationTool/Scripts/CheckXcodeVersion.cs
+++ b/Engine/Source/Programs/AutomationTool/Scripts/CheckXcodeVersion.cs
@@ -15,27 +15,22 @@ namespace AutomationTool
 	{
 		public override void ExecuteBuild()
 		{
-			string Version = ParseParamValue("Version");
-			if(Version == null)
+			string RequestedVersion = ParseParamValue("Version");
+			if(RequestedVersion == null)
 			{
 				throw new AutomationException("Missing -Version=... parameter");
 			}
 
-			IProcessResult Result = Run("xcodebuild", "-version");
-			if(Result.ExitCode != 0)
-			{
-				throw new AutomationException("Unable to query version number from xcodebuild (exit code={0})", Result.ExitCode);
-			}
+			string InstalledSdkVersion = UnrealBuildBase.ApplePlatformSDK.InstalledSDKVersion;
 
-			Match Match = Regex.Match(Result.Output, "^Xcode ([0-9.]+)", RegexOptions.Multiline);
-			if(!Match.Success)
+			if (InstalledSdkVersion == null)
 			{
-				throw new AutomationException("Missing version number from xcodebuild output:\n{0}", Result.Output);
+				throw new AutomationException("Unable to query version number from xcodebuild");
 			}
-
-			if(Match.Groups[1].Value != Version)
+			
+			if (InstalledSdkVersion != RequestedVersion)
 			{
-				LogWarning("Installed Xcode version is {0} - expected {1}", Match.Groups[1].Value, Version);
+				LogWarning("Installed Xcode version is {0} - expected {1}", InstalledSdkVersion, RequestedVersion);
 			}
 		}
 	}
diff --git a/Engine/Source/Programs/DotNETCommon/BuildUtilities/BuildUtilities.csproj b/Engine/Source/Programs/DotNETCommon/BuildUtilities/BuildUtilities.csproj
--- a/Engine/Source/Programs/DotNETCommon/BuildUtilities/BuildUtilities.csproj
+++ b/Engine/Source/Programs/DotNETCommon/BuildUtilities/BuildUtilities.csproj
@@ -31,6 +31,7 @@
   </ItemGroup>
   <ItemGroup>
     <PackageReference Include="Microsoft.CSharp" Version="4.7.0" />
+    <PackageReference Include="Microsoft.Win32.Registry" Version="4.7.0" />
     <PackageReference Include="System.Data.DataSetExtensions" Version="4.5.0" />
   </ItemGroup>
 </Project>
\ No newline at end of file
diff --git a/Engine/Source/Programs/DotNETCommon/BuildUtilities/Platform/Mac/ApplePlatformSDK.cs b/Engine/Source/Programs/DotNETCommon/BuildUtilities/Platform/Mac/ApplePlatformSDK.cs
new file mode 100644
--- /dev/null
+++ b/Engine/Source/Programs/DotNETCommon/BuildUtilities/Platform/Mac/ApplePlatformSDK.cs
@@ -0,0 +1,173 @@
+// Copyright Epic Games, Inc. All Rights Reserved.
+
+using System;
+using System.Diagnostics;
+using System.IO;
+using System.Text;
+using System.Text.RegularExpressions;
+using EpicGames.Core;
+using Microsoft.VisualBasic.CompilerServices;
+using Microsoft.Win32;
+
+namespace UnrealBuildBase
+{
+	static public class ApplePlatformSDK
+	{
+		public static readonly string? InstalledSDKVersion = GetInstalledSDKVersion();
+
+		private static string? GetInstalledSDKVersion()
+		{
+			// get xcode version on Mac
+			if (RuntimePlatform.IsMac)
+			{
+				string Output = RunLocalProcessAndReturnStdOut("sh", "-c 'xcodebuild -version'");
+				Match Result = Regex.Match(Output, @"Xcode (\S*)");
+				return Result.Success ? Result.Groups[1].Value : null;
+			}
+
+			if (RuntimePlatform.IsWindows)
+			{
+				// otherwise, get iTunes "Version"
+				string? DllPath =
+					Registry.GetValue(
+						"HKEY_LOCAL_MACHINE\\SOFTWARE\\Wow6432Node\\Apple Inc.\\Apple Mobile Device Support\\Shared",
+						"iTunesMobileDeviceDLL", null) as string;
+				if (string.IsNullOrEmpty(DllPath) || !File.Exists(DllPath))
+				{
+					DllPath = Registry.GetValue(
+						"HKEY_LOCAL_MACHINE\\SOFTWARE\\Wow6432Node\\Apple Inc.\\Apple Mobile Device Support\\Shared",
+						"MobileDeviceDLL", null) as string;
+					if (string.IsNullOrEmpty(DllPath) || !File.Exists(DllPath))
+					{
+						// iTunes >= 12.7 doesn't have a key specifying the 32-bit DLL but it does have a ASMapiInterfaceDLL key and MobileDevice.dll is in usually in the same directory
+						DllPath = Registry.GetValue(
+							"HKEY_LOCAL_MACHINE\\SOFTWARE\\Wow6432Node\\Apple Inc.\\Apple Mobile Device Support\\Shared",
+							"ASMapiInterfaceDLL", null) as string;
+						DllPath = String.IsNullOrEmpty(DllPath)
+							? null
+							: DllPath.Substring(0, DllPath.LastIndexOf('\\') + 1) + "MobileDevice.dll";
+
+						if (string.IsNullOrEmpty(DllPath) || !File.Exists(DllPath))
+						{
+							DllPath = FindWindowsStoreITunesDLL();
+						}
+					}
+				}
+
+				if (!string.IsNullOrEmpty(DllPath) && File.Exists(DllPath))
+				{
+					return FileVersionInfo.GetVersionInfo(DllPath).FileVersion;
+				}
+			}
+
+			return null;
+		}
+
+		private static string? FindWindowsStoreITunesDLL()
+		{
+			string? InstallPath = null;
+
+			string PackagesKeyName = "Software\\Classes\\Local Settings\\Software\\Microsoft\\Windows\\CurrentVersion\\AppModel\\PackageRepository\\Packages";
+
+			RegistryKey PackagesKey = Registry.LocalMachine.OpenSubKey(PackagesKeyName);
+			if (PackagesKey != null)
+			{
+				string[] PackageSubKeyNames = PackagesKey.GetSubKeyNames();
+
+				foreach (string PackageSubKeyName in PackageSubKeyNames)
+				{
+					if (PackageSubKeyName.Contains("AppleInc.iTunes") && (PackageSubKeyName.Contains("_x64") || PackageSubKeyName.Contains("_x86")))
+					{
+						string FullPackageSubKeyName = PackagesKeyName + "\\" + PackageSubKeyName;
+
+						RegistryKey iTunesKey = Registry.LocalMachine.OpenSubKey(FullPackageSubKeyName);
+						if (iTunesKey != null)
+						{
+							InstallPath = (string)iTunesKey.GetValue("Path") + "\\AMDS32\\MobileDevice.dll";
+							break;
+						}
+					}
+				}
+			}
+
+			return InstallPath;
+		}
+		
+		/// <summary>
+		/// Runs a command line process, and returns simple StdOut output. This doesn't handle errors or return codes
+		/// </summary>
+		/// <returns>The entire StdOut generated from the process as a single trimmed string</returns>
+		/// <param name="Command">Command to run</param>
+		/// <param name="Args">Arguments to Command</param>
+		private static string RunLocalProcessAndReturnStdOut(string Command, string Args)
+		{
+			int ExitCode;
+			return RunLocalProcessAndReturnStdOut(Command, Args, out ExitCode);	
+		}
+
+		/// <summary>
+		/// Runs a command line process, and returns simple StdOut output.
+		/// </summary>
+		/// <returns>The entire StdOut generated from the process as a single trimmed string</returns>
+		/// <param name="Command">Command to run</param>
+		/// <param name="Args">Arguments to Command</param>
+		/// <param name="ExitCode">The return code from the process after it exits</param>
+		/// <param name="LogOutput">Whether to also log standard output and standard error</param>
+		private static string RunLocalProcessAndReturnStdOut(string Command, string? Args, out int ExitCode, bool LogOutput = false)
+		{
+			// Process Arguments follow windows conventions in .NET Core
+			// Which means single quotes ' are not considered quotes.
+			// see [REDACTED_PROFILE_URL]/runtime/issues/29857
+			// also see UE-102580
+			// for rules see https://docs.microsoft.com/en-us/cpp/cpp/main-function-command-line-args
+			Args = Args?.Replace('\'', '\"');
+
+			ProcessStartInfo StartInfo = new ProcessStartInfo(Command, Args);
+			StartInfo.UseShellExecute = false;
+			StartInfo.RedirectStandardInput = true;
+			StartInfo.RedirectStandardOutput = true;
+			StartInfo.RedirectStandardError = true;
+			StartInfo.CreateNoWindow = true;
+			StartInfo.StandardOutputEncoding = Encoding.UTF8;
+
+			string FullOutput = "";
+			string ErrorOutput = "";
+			using (Process LocalProcess = Process.Start(StartInfo))
+			{
+				StreamReader OutputReader = LocalProcess.StandardOutput;
+				// trim off any extraneous new lines, helpful for those one-line outputs
+				FullOutput = OutputReader.ReadToEnd().Trim();
+
+				StreamReader ErrorReader = LocalProcess.StandardError;
+				// trim off any extraneous new lines, helpful for those one-line outputs
+				ErrorOutput = ErrorReader.ReadToEnd().Trim();
+				if (LogOutput)
+				{
+					if(FullOutput.Length > 0)
+					{
+						Log.TraceInformation(FullOutput);
+					}
+
+					if (ErrorOutput.Length > 0)
+					{
+						Log.TraceError(ErrorOutput);
+					}
+				}
+
+				LocalProcess.WaitForExit();
+				ExitCode = LocalProcess.ExitCode;
+			}
+
+			// trim off any extraneous new lines, helpful for those one-line outputs
+			if (ErrorOutput.Length > 0)
+			{
+				if (FullOutput.Length > 0)
+				{
+					FullOutput += Environment.NewLine;
+				}
+				FullOutput += ErrorOutput;
+			}
+			return FullOutput;
+		}
+	}
+}
\ No newline at end of file
diff --git a/Engine/Source/Programs/UnrealBuildTool/Platform/Mac/ApplePlatformSDK.cs b/Engine/Source/Programs/UnrealBuildTool/Platform/Mac/ApplePlatformSDK.cs
--- a/Engine/Source/Programs/UnrealBuildTool/Platform/Mac/ApplePlatformSDK.cs
+++ b/Engine/Source/Programs/UnrealBuildTool/Platform/Mac/ApplePlatformSDK.cs
@@ -39,45 +39,6 @@ namespace UnrealBuildTool
 			MinVersion = MaxVersion = null;
 		}
 
-		public override string GetInstalledSDKVersion()
-		{
-			// get xcode version on Mac
-			if (RuntimePlatform.IsMac)
-			{
-				string Output = Utils.RunLocalProcessAndReturnStdOut("sh", "-c 'xcodebuild -version'");
-				Match Result = Regex.Match(Output, @"Xcode (\S*)");
-				return Result.Success ? Result.Groups[1].Value : "";
-			}
-
-			if (RuntimePlatform.IsWindows)
-			{
-				// otherwise, get iTunes "Version"
-				string DllPath = Registry.GetValue("HKEY_LOCAL_MACHINE\\SOFTWARE\\Wow6432Node\\Apple Inc.\\Apple Mobile Device Support\\Shared", "iTunesMobileDeviceDLL", null) as string;
-				if (string.IsNullOrEmpty(DllPath) || !File.Exists(DllPath))
-				{
-					DllPath = Registry.GetValue("HKEY_LOCAL_MACHINE\\SOFTWARE\\Wow6432Node\\Apple Inc.\\Apple Mobile Device Support\\Shared", "MobileDeviceDLL", null) as string;
-					if (string.IsNullOrEmpty(DllPath) || !File.Exists(DllPath))
-					{
-						// iTunes >= 12.7 doesn't have a key specifying the 32-bit DLL but it does have a ASMapiInterfaceDLL key and MobileDevice.dll is in usually in the same directory
-						DllPath = Registry.GetValue("HKEY_LOCAL_MACHINE\\SOFTWARE\\Wow6432Node\\Apple Inc.\\Apple Mobile Device Support\\Shared", "ASMapiInterfaceDLL", null) as string;
-						DllPath = String.IsNullOrEmpty(DllPath) ? null : DllPath.Substring(0, DllPath.LastIndexOf('\\') + 1) + "MobileDevice.dll";
-
-						if (string.IsNullOrEmpty(DllPath) || !File.Exists(DllPath))
-						{
-							DllPath = FindWindowsStoreITunesDLL();
-						}
-					}
-				}
-
-				if (!string.IsNullOrEmpty(DllPath) && File.Exists(DllPath))
-				{
-					return FileVersionInfo.GetVersionInfo(DllPath).FileVersion;
-				}
-			}
-
-			return null;
-		}
-
 		public override bool TryConvertVersionToInt(string StringValue, out UInt64 OutValue)
 		{
 			// 8 bits per component, with high getting extra from high 32
@@ -99,38 +60,9 @@ namespace UnrealBuildTool
 			return false;
 		}
 
-
-
-
-
-		static string FindWindowsStoreITunesDLL()
+		public override string GetInstalledSDKVersion()
 		{
-			string InstallPath = null;
-
-			string PackagesKeyName = "Software\\Classes\\Local Settings\\Software\\Microsoft\\Windows\\CurrentVersion\\AppModel\\PackageRepository\\Packages";
-
-			RegistryKey PackagesKey = Registry.LocalMachine.OpenSubKey(PackagesKeyName);
-			if (PackagesKey != null)
-			{
-				string[] PackageSubKeyNames = PackagesKey.GetSubKeyNames();
-
-				foreach (string PackageSubKeyName in PackageSubKeyNames)
-				{
-					if (PackageSubKeyName.Contains("AppleInc.iTunes") && (PackageSubKeyName.Contains("_x64") || PackageSubKeyName.Contains("_x86")))
-					{
-						string FullPackageSubKeyName = PackagesKeyName + "\\" + PackageSubKeyName;
-
-						RegistryKey iTunesKey = Registry.LocalMachine.OpenSubKey(FullPackageSubKeyName);
-						if (iTunesKey != null)
-						{
-							InstallPath = (string)iTunesKey.GetValue("Path") + "\\AMDS32\\MobileDevice.dll";
-							break;
-						}
-					}
-				}
-			}
-
-			return InstallPath;
+			return UnrealBuildBase.ApplePlatformSDK.InstalledSDKVersion;
 		}
 	}
 }
diff --git a/Engine/Source/Programs/UnrealBuildTool/ToolChain/AppleToolChain.cs b/Engine/Source/Programs/UnrealBuildTool/ToolChain/AppleToolChain.cs
--- a/Engine/Source/Programs/UnrealBuildTool/ToolChain/AppleToolChain.cs
+++ b/Engine/Source/Programs/UnrealBuildTool/ToolChain/AppleToolChain.cs
@@ -58,25 +58,14 @@ namespace UnrealBuildTool
 			// Installed engine requires Xcode 11
 			if (Unreal.IsEngineInstalled())
 			{
-				string XcodeBuilderVersionOutput = Utils.RunLocalProcessAndReturnStdOut("xcodebuild", "-version");
-				if (XcodeBuilderVersionOutput.Length > 10)
+				string InstalledSdkVersion = UnrealBuildBase.ApplePlatformSDK.InstalledSDKVersion;
+				if (String.IsNullOrEmpty(InstalledSdkVersion))
 				{
-					string[] Version = XcodeBuilderVersionOutput.Substring(6, 4).Split('.');
-					if (Version.Length == 2)
-					{
-						if (int.Parse(Version[0]) < 11)
-						{
-							throw new BuildException("Building for macOS, iOS and tvOS requires Xcode 11 or newer, Xcode " + Version[0] + "." + Version[1] + " detected");
-						}
-					}
-					else
-					{
-						Log.TraceWarning("Failed to query Xcode version");
-					}
+					throw new BuildException("Unable to get xcode version");
 				}
-				else
+				if (int.Parse(InstalledSdkVersion.Substring(0,2)) < 11)
 				{
-					Log.TraceWarning("Failed to query Xcode version");
+					throw new BuildException("Building for macOS, iOS and tvOS requires Xcode 11 or newer, Xcode " + InstalledSdkVersion + " detected");
 				}
 			}
 		}
