BuildUtilities: Unreal.GetExtensionDirs()

Reverse filter direction - keep existing platforms

(fix thanks to [REDACTED_PERSON])

diff --git a/Engine/Source/Programs/DotNETCommon/BuildUtilities/Unreal.cs b/Engine/Source/Programs/DotNETCommon/BuildUtilities/Unreal.cs
--- a/Engine/Source/Programs/DotNETCommon/BuildUtilities/Unreal.cs
+++ b/Engine/Source/Programs/DotNETCommon/BuildUtilities/Unreal.cs
@@ -234,7 +234,7 @@ namespace UnrealBuildBase
 				if (BaseDir != Unreal.EngineDirectory && CachedDirs.Item1.Count > 0)
 				{
 					// if the DDPI.ini file doesn't exist, we haven't synced the platform, so just skip this directory
-					CachedDirs.Item1.RemoveAll(x => DataDrivenPlatformInfoIniIsPresent(x.GetDirectoryName()));
+					CachedDirs.Item1.RemoveAll(x => !DataDrivenPlatformInfoIniIsPresent(x.GetDirectoryName()));
 				}
 			}
 
